Add texture PCX translator.

diff --git a/Engine/Plugins/Runtime/Interchange/Source/Import/Private/InterchangeImportPlugin.cpp b/Engine/Plugins/Runtime/Interchange/Source/Import/Private/InterchangeImportPlugin.cpp
--- a/Engine/Plugins/Runtime/Interchange/Source/Import/Private/InterchangeImportPlugin.cpp
+++ b/Engine/Plugins/Runtime/Interchange/Source/Import/Private/InterchangeImportPlugin.cpp
@@ -4,9 +4,11 @@
 #include "IInterchangeImportPlugin.h"
 #include "InterchangeManager.h"
 #include "Modules/ModuleManager.h"
-#include "Texture/InterchangeTextureFactory.h"
-#include "Texture/InterchangePNGTranslator.h"
 #include "Texture/InterchangeJPGTranslator.h"
+#include "Texture/InterchangePCXTranslator.h"
+#include "Texture/InterchangePNGTranslator.h"
+#include "Texture/InterchangeTextureFactory.h"
+
 
 DEFINE_LOG_CATEGORY(LogInterchangeImportPlugin);
 
@@ -26,8 +28,11 @@ void FInterchangeImportPlugin::StartupModule()
 	UInterchangeManager& InterchangeManager = UInterchangeManager::GetInterchangeManager();
 
 	//Register the translators
+
+	//Textures
 	InterchangeManager.RegisterTranslator(UInterchangePNGTranslator::StaticClass());
 	InterchangeManager.RegisterTranslator(UInterchangeJPGTranslator::StaticClass());
+	InterchangeManager.RegisterTranslator(UInterchangePCXTranslator::StaticClass());
 
 	//Register the factories
 	InterchangeManager.RegisterFactory(UInterchangeTextureFactory::StaticClass());
diff --git a/Engine/Plugins/Runtime/Interchange/Source/Import/Private/Texture/InterchangePCXTranslator.cpp b/Engine/Plugins/Runtime/Interchange/Source/Import/Private/Texture/InterchangePCXTranslator.cpp
new file mode 100644
--- /dev/null
+++ b/Engine/Plugins/Runtime/Interchange/Source/Import/Private/Texture/InterchangePCXTranslator.cpp
@@ -0,0 +1,257 @@
+// Copyright Epic Games, Inc. All Rights Reserved. 
+#include "Texture/InterchangePCXTranslator.h"
+
+#include "Engine/Texture.h"
+#include "Engine/Texture2D.h"
+#include "IImageWrapper.h"
+#include "IImageWrapperModule.h"
+#include "LogInterchangeImportPlugin.h"
+#include "Misc/ConfigCacheIni.h"
+#include "Misc/FileHelper.h"
+#include "Misc/Paths.h"
+#include "Misc/ScopedSlowTask.h"
+#include "Modules/ModuleManager.h"
+#include "Nodes/BaseNodeContainer.h"
+#include "TextureNode.h"
+
+
+//////////////////////////////////////////////////////////////////////////
+// PCX helper local function
+
+namespace PCXParserHelper
+{
+	// .PCX file header.
+#pragma pack(push,1)
+	
+	class FPCXFileHeader
+	{
+	public:
+		uint8	Manufacturer;		// Always 10.
+		uint8	Version;			// PCX file version.
+		uint8	Encoding;			// 1=run-length, 0=none.
+		uint8	BitsPerPixel;		// 1,2,4, or 8.
+		uint16	XMin;				// Dimensions of the image.
+		uint16	YMin;				// Dimensions of the image.
+		uint16	XMax;				// Dimensions of the image.
+		uint16	YMax;				// Dimensions of the image.
+		uint16	XDotsPerInch;		// Horizontal printer resolution.
+		uint16	YDotsPerInch;		// Vertical printer resolution.
+		uint8	OldColorMap[48];	// Old colormap info data.
+		uint8	Reserved1;			// Must be 0.
+		uint8	NumPlanes;			// Number of color planes (1, 3, 4, etc).
+		uint16	BytesPerLine;		// Number of bytes per scanline.
+		uint16	PaletteType;		// How to interpret palette: 1=color, 2=gray.
+		uint16	HScreenSize;		// Horizontal monitor size.
+		uint16	VScreenSize;		// Vertical monitor size.
+		uint8	Reserved2[54];		// Must be 0.
+		friend FArchive& operator<<(FArchive& Ar, FPCXFileHeader& H)
+		{
+			Ar << H.Manufacturer << H.Version << H.Encoding << H.BitsPerPixel;
+			Ar << H.XMin << H.YMin << H.XMax << H.YMax << H.XDotsPerInch << H.YDotsPerInch;
+			for (int32 i = 0; i < UE_ARRAY_COUNT(H.OldColorMap); i++)
+				Ar << H.OldColorMap[i];
+			Ar << H.Reserved1 << H.NumPlanes;
+			Ar << H.BytesPerLine << H.PaletteType << H.HScreenSize << H.VScreenSize;
+			for (int32 i = 0; i < UE_ARRAY_COUNT(H.Reserved2); i++)
+				Ar << H.Reserved2[i];
+			return Ar;
+		}
+	};
+
+	#pragma pack(pop)
+} //PCXParserHelper namespace end
+
+bool UInterchangePCXTranslator::CanImportSourceData(const UInterchangeSourceData* InSourceData) const
+{
+	FString Extension = FPaths::GetExtension(InSourceData->GetFilename());
+	FString PCXExtension = (TEXT("pcx;Texture"));
+	return PCXExtension.StartsWith(Extension);
+}
+
+bool UInterchangePCXTranslator::Translate(const UInterchangeSourceData* SourceData, Interchange::FBaseNodeContainer& BaseNodeContainer) const
+{
+	FString Filename = SourceData->GetFilename();
+	if (!FPaths::FileExists(Filename))
+	{
+		return false;
+	}
+
+	FName DisplayLabel = *FPaths::GetBaseFilename(Filename);
+	Interchange::FNodeUniqueID NodeUID(*Filename);
+	//PCX is creating a UTexture2D
+	Interchange::FTextureNode* TextureNode = new Interchange::FTextureNode(NodeUID, DisplayLabel, UTexture2D::StaticClass());
+	TextureNode->SetPayLoadKey(Filename);
+
+	//Test node change
+ 	//TextureNode->SetCustomLODGroup((uint8)TextureGroup::TEXTUREGROUP_WorldNormalMap);
+ 	//TextureNode->SetCustomCompressionSettings((uint8)TextureCompressionSettings::TC_Normalmap);
+
+	BaseNodeContainer.AddNode(static_cast<Interchange::FBaseNode*>(TextureNode));
+	return true;
+}
+
+const TOptional<Interchange::FImportImage> UInterchangePCXTranslator::GetPayloadData(const UInterchangeSourceData* SourceData, const FString& PayLoadKey) const
+{
+	if (!SourceData)
+	{
+		UE_LOG(LogInterchangeImportPlugin, Error, TEXT("Failed to import PCX, bad source data."));
+		return TOptional<Interchange::FImportImage>();
+	}
+
+	TArray64<uint8> SourceDataBuffer;
+	FString Filename = SourceData->GetFilename();
+	
+	//Make sure the key fit the filename, The key should always be valid
+	if (!Filename.Equals(PayLoadKey))
+	{
+		UE_LOG(LogInterchangeImportPlugin, Error, TEXT("Failed to import PCX, wrong payload key. [%s]"), *Filename);
+		return TOptional<Interchange::FImportImage>();
+	}
+
+	if (!FPaths::FileExists(Filename))
+	{
+		UE_LOG(LogInterchangeImportPlugin, Error, TEXT("Failed to import PCX, cannot open file. [%s]"), *Filename);
+		return TOptional<Interchange::FImportImage>();
+	}
+
+	if (!FFileHelper::LoadFileToArray(SourceDataBuffer, *Filename))
+	{
+		UE_LOG(LogInterchangeImportPlugin, Error, TEXT("Failed to import PCX, cannot load file content into an array. [%s]"), *Filename);
+		return TOptional<Interchange::FImportImage>();
+	}
+
+	const uint8* Buffer = SourceDataBuffer.GetData();
+	const uint8* BufferEnd = Buffer + SourceDataBuffer.Num();
+
+	bool bAllowNonPowerOfTwo = false;
+	GConfig->GetBool(TEXT("TextureImporter"), TEXT("AllowNonPowerOfTwoTextures"), bAllowNonPowerOfTwo, GEditorIni);
+
+	// Validate it.
+	const int32 Length = BufferEnd - Buffer;
+
+	IImageWrapperModule& ImageWrapperModule = FModuleManager::LoadModuleChecked<IImageWrapperModule>(FName("ImageWrapper"));
+
+	Interchange::FImportImage PayloadData;
+
+	//
+	// PCX
+	//
+	const PCXParserHelper::FPCXFileHeader* PCX = (PCXParserHelper::FPCXFileHeader*)Buffer;
+	if (Length >= sizeof(PCXParserHelper::FPCXFileHeader) && PCX->Manufacturer == 10)
+	{
+		int32 NewU = PCX->XMax + 1 - PCX->XMin;
+		int32 NewV = PCX->YMax + 1 - PCX->YMin;
+
+		// Check the resolution of the imported texture to ensure validity
+		if (!Interchange::FImportImageHelper::IsImportResolutionValid(NewU, NewV, bAllowNonPowerOfTwo))
+		{
+			UE_LOG(LogInterchangeImportPlugin, Error, TEXT("Failed to import PCX, invalid resolution. Resolution[%d, %d], AllowPowerOfTwo[%s], [%s]"), NewU, NewV, bAllowNonPowerOfTwo ? TEXT("True") : TEXT("false"), *Filename);
+			return TOptional<Interchange::FImportImage>();
+		}
+		else if (PCX->NumPlanes == 1 && PCX->BitsPerPixel == 8)
+		{
+
+			// Set texture properties.
+			PayloadData.Init2DWithOneMip(
+				NewU,
+				NewV,
+				TSF_BGRA8
+			);
+			FColor* DestPtr = (FColor*)PayloadData.RawData.GetData();
+
+			// Import the palette.
+			uint8* PCXPalette = (uint8*)(Buffer + Length - 256 * 3);
+			TArray<FColor>	Palette;
+			for (uint32 i = 0; i < 256; i++)
+			{
+				Palette.Add(FColor(PCXPalette[i * 3 + 0], PCXPalette[i * 3 + 1], PCXPalette[i * 3 + 2], i == 0 ? 0 : 255));
+			}
+
+			// Import it.
+			FColor* DestEnd = DestPtr + NewU * NewV;
+			Buffer += 128;
+			while (DestPtr < DestEnd)
+			{
+				uint8 Color = *Buffer++;
+				if ((Color & 0xc0) == 0xc0)
+				{
+					uint32 RunLength = Color & 0x3f;
+					Color = *Buffer++;
+
+					for (uint32 Index = 0; Index < RunLength; Index++)
+					{
+						*DestPtr++ = Palette[Color];
+					}
+				}
+				else *DestPtr++ = Palette[Color];
+			}
+		}
+		else if (PCX->NumPlanes == 3 && PCX->BitsPerPixel == 8)
+		{
+			// Set texture properties.
+			PayloadData.Init2DWithOneMip(
+				NewU,
+				NewV,
+				TSF_BGRA8
+			);
+
+			uint8* Dest = PayloadData.RawData.GetData();
+
+			// Doing a memset to make sure the alpha channel is set to 0xff since we only have 3 color planes.
+			FMemory::Memset(Dest, 0xff, NewU * NewV * FTextureSource::GetBytesPerPixel(PayloadData.Format));
+
+			// Copy upside-down scanlines.
+			Buffer += 128;
+			int32 CountU = FMath::Min<int32>(PCX->BytesPerLine, NewU);
+			for (int32 i = 0; i < NewV; i++)
+			{
+				// We need to decode image one line per time building RGB image color plane by color plane.
+				int32 RunLength, Overflow = 0;
+				uint8 Color = 0;
+				for (int32 ColorPlane = 2; ColorPlane >= 0; ColorPlane--)
+				{
+					for (int32 j = 0; j < CountU; j++)
+					{
+						if (!Overflow)
+						{
+							Color = *Buffer++;
+							if ((Color & 0xc0) == 0xc0)
+							{
+								RunLength = FMath::Min((Color & 0x3f), CountU - j);
+								Overflow = (Color & 0x3f) - RunLength;
+								Color = *Buffer++;
+							}
+							else
+								RunLength = 1;
+						}
+						else
+						{
+							RunLength = FMath::Min(Overflow, CountU - j);
+							Overflow = Overflow - RunLength;
+						}
+
+						//checkf(((i*NewU + RunLength) * 4 + ColorPlane) < (Texture->Source.CalcMipSize(0)),
+						//	TEXT("RLE going off the end of buffer"));
+						for (int32 k = j; k < j + RunLength; k++)
+						{
+							Dest[(i * NewU + k) * 4 + ColorPlane] = Color;
+						}
+						j += RunLength - 1;
+					}
+				}
+			}
+		}
+		else
+		{
+			UE_LOG(LogInterchangeImportPlugin, Error, TEXT("Failed to import PCX, unsupported format. (%i/%i) [%s]"), PCX->NumPlanes, PCX->BitsPerPixel, *Filename);
+			return TOptional<Interchange::FImportImage>();
+		}
+	}
+	else
+	{
+		UE_LOG(LogInterchangeImportPlugin, Error, TEXT("Failed to import PCX, unsupported file version. [%s]"), *Filename);
+		return TOptional<Interchange::FImportImage>();
+	}
+
+	return PayloadData;
+}
\ No newline at end of file
diff --git a/Engine/Plugins/Runtime/Interchange/Source/Import/Public/Texture/InterchangePCXTranslator.h b/Engine/Plugins/Runtime/Interchange/Source/Import/Public/Texture/InterchangePCXTranslator.h
new file mode 100644
--- /dev/null
+++ b/Engine/Plugins/Runtime/Interchange/Source/Import/Public/Texture/InterchangePCXTranslator.h
@@ -0,0 +1,45 @@
+// Copyright Epic Games, Inc. All Rights Reserved.
+
+#pragma once
+
+#include "CoreMinimal.h"
+#include "UObject/Object.h"
+#include "UObject/ObjectMacros.h"
+#include "InterchangeTranslatorBase.h"
+#include "Texture/InterchangeTexturePayloadData.h"
+#include "Texture/InterchangeTextureTranslator.h"
+
+#include "InterchangePCXTranslator.generated.h"
+
+UCLASS(BlueprintType)
+class INTERCHANGEIMPORTPLUGIN_API UInterchangePCXTranslator : public UInterchangeTextureTranslator
+{
+	GENERATED_BODY()
+public:
+	
+	/*
+	 * return true if the translator can translate the specified file.
+	 */
+	virtual bool CanImportSourceData(const UInterchangeSourceData* InSourceData) const override;
+
+	/**
+	 * Translate a source data into a node hold by the specified nodes container.
+	 *
+	 * @param SourceData - The source data containing the data to translate
+	 * @param BaseNodeContainer - The unreal objects descriptions container where to put the translated source data.
+	 * @return true if the translator can translate the source data, false otherwise.
+	 */
+	virtual bool Translate(const UInterchangeSourceData* SourceData, Interchange::FBaseNodeContainer& BaseNodeContainer) const override;
+
+	/**
+	 * Once the translation is done, the import process need a way to retrieve payload data.
+	 * This payload will be use by the factories to create the asset.
+	 *
+	 * @param SourceData - The source data containing the data to translate
+	 * @param PayloadKey - The key to retrieve the a particular payload contain into the specified source data.
+	 * @return a PayloadData containing the import image data. The TOptional will not be set if there is an error.
+	 */
+	virtual const TOptional<Interchange::FImportImage> GetPayloadData(const UInterchangeSourceData* SourceData, const FString& PayLoadKey) const override;
+};
+
+
