Changed API so that FBaseNodeContainer::AddNode accepts a TUniquePtr instead of a raw pointer.
Added BMP Translator.

diff --git a/Engine/Plugins/Runtime/Interchange/Source/Import/Private/InterchangeImportPlugin.cpp b/Engine/Plugins/Runtime/Interchange/Source/Import/Private/InterchangeImportPlugin.cpp
--- a/Engine/Plugins/Runtime/Interchange/Source/Import/Private/InterchangeImportPlugin.cpp
+++ b/Engine/Plugins/Runtime/Interchange/Source/Import/Private/InterchangeImportPlugin.cpp
@@ -4,6 +4,7 @@
 #include "IInterchangeImportPlugin.h"
 #include "InterchangeManager.h"
 #include "Modules/ModuleManager.h"
+#include "Texture/InterchangeBMPTranslator.h"
 #include "Texture/InterchangeJPGTranslator.h"
 #include "Texture/InterchangePCXTranslator.h"
 #include "Texture/InterchangePNGTranslator.h"
@@ -30,9 +31,10 @@ void FInterchangeImportPlugin::StartupModule()
 	//Register the translators
 
 	//Textures
-	InterchangeManager.RegisterTranslator(UInterchangePNGTranslator::StaticClass());
+	InterchangeManager.RegisterTranslator(UInterchangeBMPTranslator::StaticClass());
 	InterchangeManager.RegisterTranslator(UInterchangeJPGTranslator::StaticClass());
 	InterchangeManager.RegisterTranslator(UInterchangePCXTranslator::StaticClass());
+	InterchangeManager.RegisterTranslator(UInterchangePNGTranslator::StaticClass());
 
 	//Register the factories
 	InterchangeManager.RegisterFactory(UInterchangeTextureFactory::StaticClass());
diff --git a/Engine/Plugins/Runtime/Interchange/Source/Import/Private/Texture/InterchangeBMPTranslator.cpp b/Engine/Plugins/Runtime/Interchange/Source/Import/Private/Texture/InterchangeBMPTranslator.cpp
new file mode 100644
--- /dev/null
+++ b/Engine/Plugins/Runtime/Interchange/Source/Import/Private/Texture/InterchangeBMPTranslator.cpp
@@ -0,0 +1,120 @@
+// Copyright Epic Games, Inc. All Rights Reserved. 
+
+#include "Texture/InterchangeBMPTranslator.h"
+
+#include "Engine/Texture.h"
+#include "Engine/Texture2D.h"
+#include "IImageWrapper.h"
+#include "IImageWrapperModule.h"
+#include "LogInterchangeImportPlugin.h"
+#include "Misc/ConfigCacheIni.h"
+#include "Misc/FileHelper.h"
+#include "Misc/Paths.h"
+#include "Misc/ScopedSlowTask.h"
+#include "Modules/ModuleManager.h"
+#include "Nodes/BaseNodeContainer.h"
+#include "TextureNode.h"
+
+
+
+bool UInterchangeBMPTranslator::CanImportSourceData(const UInterchangeSourceData* InSourceData) const
+{
+	FString Extension = FPaths::GetExtension(InSourceData->GetFilename());
+	FString BMPExtension = (TEXT("bmp;Texture"));
+	return BMPExtension.StartsWith(Extension);
+}
+
+bool UInterchangeBMPTranslator::Translate(const UInterchangeSourceData* SourceData, Interchange::FBaseNodeContainer& BaseNodeContainer) const
+{
+	FString Filename = SourceData->GetFilename();
+	if (!FPaths::FileExists(Filename))
+	{
+		return false;
+	}
+
+	FName DisplayLabel = *FPaths::GetBaseFilename(Filename);
+	Interchange::FNodeUniqueID NodeUID(*Filename);
+	// BMP is creating a UTexture2D
+	TUniquePtr<Interchange::FTextureNode> TextureNode = MakeUnique<Interchange::FTextureNode>(NodeUID, DisplayLabel, UTexture2D::StaticClass());
+	TextureNode->SetPayLoadKey(Filename);
+
+	BaseNodeContainer.AddNode(MoveTemp(TextureNode));
+	return true;
+}
+
+const TOptional<Interchange::FImportImage> UInterchangeBMPTranslator::GetPayloadData(const UInterchangeSourceData* SourceData, const FString& PayLoadKey) const
+{
+	if (!SourceData)
+	{
+		UE_LOG(LogInterchangeImportPlugin, Error, TEXT("Failed to import BMP, bad source data."));
+		return TOptional<Interchange::FImportImage>();
+	}
+
+	TArray64<uint8> SourceDataBuffer;
+	FString Filename = SourceData->GetFilename();
+
+	//Make sure the key fit the filename, The key should always be valid
+	if (!Filename.Equals(PayLoadKey))
+	{
+		UE_LOG(LogInterchangeImportPlugin, Error, TEXT("Failed to import BMP, wrong payload key. [%s]"), *Filename);
+		return TOptional<Interchange::FImportImage>();
+	}
+
+	if (!FPaths::FileExists(Filename))
+	{
+		UE_LOG(LogInterchangeImportPlugin, Error, TEXT("Failed to import BMP, cannot open file. [%s]"), *Filename);
+		return TOptional<Interchange::FImportImage>();
+	}
+
+	if (!FFileHelper::LoadFileToArray(SourceDataBuffer, *Filename))
+	{
+		UE_LOG(LogInterchangeImportPlugin, Error, TEXT("Failed to import BMP, cannot load file content into an array. [%s]"), *Filename);
+		return TOptional<Interchange::FImportImage>();
+	}
+
+	const uint8* Buffer = SourceDataBuffer.GetData();
+	const uint8* BufferEnd = Buffer + SourceDataBuffer.Num();
+
+	bool bAllowNonPowerOfTwo = false;
+	GConfig->GetBool(TEXT("TextureImporter"), TEXT("AllowNonPowerOfTwoTextures"), bAllowNonPowerOfTwo, GEditorIni);
+
+	// Validate it.
+	const int32 Length = BufferEnd - Buffer;
+
+	IImageWrapperModule& ImageWrapperModule = FModuleManager::LoadModuleChecked<IImageWrapperModule>(FName("ImageWrapper"));
+
+	//
+	// BMP
+	//
+	TSharedPtr<IImageWrapper> BmpImageWrapper = ImageWrapperModule.CreateImageWrapper(EImageFormat::BMP);
+	if (!BmpImageWrapper.IsValid() || !BmpImageWrapper->SetCompressed(Buffer, Length))
+	{
+		UE_LOG(LogInterchangeImportPlugin, Error, TEXT("Failed to decode BMP. [%s]"), *Filename);
+		return TOptional<Interchange::FImportImage>();
+	}
+	if (!Interchange::FImportImageHelper::IsImportResolutionValid(BmpImageWrapper->GetWidth(), BmpImageWrapper->GetHeight(), bAllowNonPowerOfTwo))
+	{
+		UE_LOG(LogInterchangeImportPlugin, Error, TEXT("Failed to import BMP, invalid resolution. Resolution[%d, %d], AllowPowerOfTwo[%s], [%s]"), BmpImageWrapper->GetWidth(), BmpImageWrapper->GetHeight(), bAllowNonPowerOfTwo ? TEXT("True") : TEXT("false"), *Filename);
+		return TOptional<Interchange::FImportImage>();
+	}
+
+	int32 BitDepth = BmpImageWrapper->GetBitDepth();
+	ERGBFormat Format = BmpImageWrapper->GetFormat();
+
+	Interchange::FImportImage PayloadData;
+
+	PayloadData.Init2DWithParams(
+		BmpImageWrapper->GetWidth(),
+		BmpImageWrapper->GetHeight(),
+		TSF_BGRA8,
+		false
+	);
+
+	if (!BmpImageWrapper->GetRaw(Format, BitDepth, PayloadData.RawData))
+	{
+		UE_LOG(LogInterchangeImportPlugin, Error, TEXT("Failed to decode BMP. [%s]"), *Filename);
+		return TOptional<Interchange::FImportImage>();
+	}
+
+	return PayloadData;
+}
diff --git a/Engine/Plugins/Runtime/Interchange/Source/Import/Private/Texture/InterchangeJPGTranslator.cpp b/Engine/Plugins/Runtime/Interchange/Source/Import/Private/Texture/InterchangeJPGTranslator.cpp
--- a/Engine/Plugins/Runtime/Interchange/Source/Import/Private/Texture/InterchangeJPGTranslator.cpp
+++ b/Engine/Plugins/Runtime/Interchange/Source/Import/Private/Texture/InterchangeJPGTranslator.cpp
@@ -32,14 +32,14 @@ bool UInterchangeJPGTranslator::Translate(const UInterchangeSourceData* SourceDa
 	FName DisplayLabel = *FPaths::GetBaseFilename(Filename);
 	Interchange::FNodeUniqueID NodeUID(*Filename);
 	//PNG is creating a UTexture2D
-	Interchange::FTextureNode* TextureNode = new Interchange::FTextureNode(NodeUID, DisplayLabel, UTexture2D::StaticClass());
+	TUniquePtr<Interchange::FTextureNode> TextureNode = MakeUnique<Interchange::FTextureNode>(NodeUID, DisplayLabel, UTexture2D::StaticClass());
 	TextureNode->SetPayLoadKey(Filename);
 
 	//Test node change
  	//TextureNode->SetCustomLODGroup((uint8)TextureGroup::TEXTUREGROUP_WorldNormalMap);
  	//TextureNode->SetCustomCompressionSettings((uint8)TextureCompressionSettings::TC_Normalmap);
 
-	BaseNodeContainer.AddNode(static_cast<Interchange::FBaseNode*>(TextureNode));
+	BaseNodeContainer.AddNode(MoveTemp(TextureNode));
 	return true;
 }
 
diff --git a/Engine/Plugins/Runtime/Interchange/Source/Import/Private/Texture/InterchangePCXTranslator.cpp b/Engine/Plugins/Runtime/Interchange/Source/Import/Private/Texture/InterchangePCXTranslator.cpp
--- a/Engine/Plugins/Runtime/Interchange/Source/Import/Private/Texture/InterchangePCXTranslator.cpp
+++ b/Engine/Plugins/Runtime/Interchange/Source/Import/Private/Texture/InterchangePCXTranslator.cpp
@@ -79,14 +79,14 @@ bool UInterchangePCXTranslator::Translate(const UInterchangeSourceData* SourceDa
 	FName DisplayLabel = *FPaths::GetBaseFilename(Filename);
 	Interchange::FNodeUniqueID NodeUID(*Filename);
 	//PCX is creating a UTexture2D
-	Interchange::FTextureNode* TextureNode = new Interchange::FTextureNode(NodeUID, DisplayLabel, UTexture2D::StaticClass());
+	TUniquePtr<Interchange::FTextureNode> TextureNode = MakeUnique<Interchange::FTextureNode>(NodeUID, DisplayLabel, UTexture2D::StaticClass());
 	TextureNode->SetPayLoadKey(Filename);
 
 	//Test node change
  	//TextureNode->SetCustomLODGroup((uint8)TextureGroup::TEXTUREGROUP_WorldNormalMap);
  	//TextureNode->SetCustomCompressionSettings((uint8)TextureCompressionSettings::TC_Normalmap);
 
-	BaseNodeContainer.AddNode(static_cast<Interchange::FBaseNode*>(TextureNode));
+	BaseNodeContainer.AddNode(MoveTemp(TextureNode));
 	return true;
 }
 
diff --git a/Engine/Plugins/Runtime/Interchange/Source/Import/Private/Texture/InterchangePNGTranslator.cpp b/Engine/Plugins/Runtime/Interchange/Source/Import/Private/Texture/InterchangePNGTranslator.cpp
--- a/Engine/Plugins/Runtime/Interchange/Source/Import/Private/Texture/InterchangePNGTranslator.cpp
+++ b/Engine/Plugins/Runtime/Interchange/Source/Import/Private/Texture/InterchangePNGTranslator.cpp
@@ -199,14 +199,14 @@ bool UInterchangePNGTranslator::Translate(const UInterchangeSourceData* SourceDa
 	FName DisplayLabel = *FPaths::GetBaseFilename(Filename);
 	Interchange::FNodeUniqueID NodeUID(*Filename);
 	//PNG is creating a UTexture2D
-	Interchange::FTextureNode* TextureNode = new Interchange::FTextureNode(NodeUID, DisplayLabel, UTexture2D::StaticClass());
+	TUniquePtr<Interchange::FTextureNode> TextureNode = MakeUnique<Interchange::FTextureNode>(NodeUID, DisplayLabel, UTexture2D::StaticClass());
 	TextureNode->SetPayLoadKey(Filename);
 
 	//Test node change
  	//TextureNode->SetCustomLODGroup((uint8)TextureGroup::TEXTUREGROUP_WorldNormalMap);
  	//TextureNode->SetCustomCompressionSettings((uint8)TextureCompressionSettings::TC_Normalmap);
 
-	BaseNodeContainer.AddNode(static_cast<Interchange::FBaseNode*>(TextureNode));
+	BaseNodeContainer.AddNode(MoveTemp(TextureNode));
 	return true;
 }
 
@@ -329,7 +329,7 @@ const TOptional<Interchange::FImportImage> UInterchangePNGTranslator::GetPayload
 	}
 	else
 	{
-		UE_LOG(LogInterchangeImportPlugin, Error, TEXT("Failed to decode JPEG. [%s]"), *Filename);
+		UE_LOG(LogInterchangeImportPlugin, Error, TEXT("Failed to decode PNG. [%s]"), *Filename);
 		return TOptional<Interchange::FImportImage>();
 	}
 	return PayloadData;
diff --git a/Engine/Plugins/Runtime/Interchange/Source/Import/Public/Texture/InterchangeBMPTranslator.h b/Engine/Plugins/Runtime/Interchange/Source/Import/Public/Texture/InterchangeBMPTranslator.h
new file mode 100644
--- /dev/null
+++ b/Engine/Plugins/Runtime/Interchange/Source/Import/Public/Texture/InterchangeBMPTranslator.h
@@ -0,0 +1,43 @@
+// Copyright Epic Games, Inc. All Rights Reserved.
+
+#pragma once
+
+#include "CoreMinimal.h"
+#include "UObject/Object.h"
+#include "UObject/ObjectMacros.h"
+#include "InterchangeTranslatorBase.h"
+#include "Texture/InterchangeTexturePayloadData.h"
+#include "Texture/InterchangeTextureTranslator.h"
+
+#include "InterchangeBMPTranslator.generated.h"
+
+UCLASS(BlueprintType)
+class INTERCHANGEIMPORTPLUGIN_API UInterchangeBMPTranslator : public UInterchangeTextureTranslator
+{
+	GENERATED_BODY()
+public:
+
+	/*
+	 * return true if the translator can translate the specified file.
+	 */
+	virtual bool CanImportSourceData(const UInterchangeSourceData* InSourceData) const override;
+
+	/**
+	 * Translate a source data into a node hold by the specified nodes container.
+	 *
+	 * @param SourceData - The source data containing the data to translate
+	 * @param BaseNodeContainer - The unreal objects descriptions container where to put the translated source data.
+	 * @return true if the translator can translate the source data, false otherwise.
+	 */
+	virtual bool Translate(const UInterchangeSourceData* SourceData, Interchange::FBaseNodeContainer& BaseNodeContainer) const override;
+
+	/**
+	 * Once the translation is done, the import process need a way to retrieve payload data.
+	 * This payload will be use by the factories to create the asset.
+	 *
+	 * @param SourceData - The source data containing the data to translate
+	 * @param PayloadKey - The key to retrieve the a particular payload contain into the specified source data.
+	 * @return a PayloadData containing the import image data. The TOptional will not be set if there is an error.
+	 */
+	virtual const TOptional<Interchange::FImportImage> GetPayloadData(const UInterchangeSourceData* SourceData, const FString& PayLoadKey) const override;
+};
diff --git a/Engine/Source/Runtime/Interchange/Core/Public/Nodes/BaseNodeContainer.h b/Engine/Source/Runtime/Interchange/Core/Public/Nodes/BaseNodeContainer.h
--- a/Engine/Source/Runtime/Interchange/Core/Public/Nodes/BaseNodeContainer.h
+++ b/Engine/Source/Runtime/Interchange/Core/Public/Nodes/BaseNodeContainer.h
@@ -37,9 +37,9 @@ public:
 	 * @note The node memory will be manage by the container after the add, the destructor of the container will free correctly all nodes
 	 *
 	 */
-	FNodeUniqueID AddNode(FBaseNode* Node)
+	FNodeUniqueID AddNode(TUniquePtr<FBaseNode> Node)
 	{
-		if (!Node)
+		if (!Node.IsValid())
 		{
 			return FBaseNode::InvalidNodeUID();
 		}
@@ -56,7 +56,7 @@ public:
 			return FBaseNode::InvalidNodeUID();
 		}
 		//Create a unique node pointer
-		Nodes.Add(NodeUniqueID, TUniquePtr<FBaseNode>(Node));
+		Nodes.Add(NodeUniqueID, MoveTemp(Node));
 		return NodeUniqueID;
 	}
 
