Add EXR texture format to interchange

diff --git a/Engine/Plugins/Runtime/Interchange/Source/Import/Private/InterchangeImportPlugin.cpp b/Engine/Plugins/Runtime/Interchange/Source/Import/Private/InterchangeImportPlugin.cpp
--- a/Engine/Plugins/Runtime/Interchange/Source/Import/Private/InterchangeImportPlugin.cpp
+++ b/Engine/Plugins/Runtime/Interchange/Source/Import/Private/InterchangeImportPlugin.cpp
@@ -5,6 +5,7 @@
 #include "InterchangeManager.h"
 #include "Modules/ModuleManager.h"
 #include "Texture/InterchangeBMPTranslator.h"
+#include "Texture/InterchangeEXRTranslator.h"
 #include "Texture/InterchangeJPGTranslator.h"
 #include "Texture/InterchangePCXTranslator.h"
 #include "Texture/InterchangePNGTranslator.h"
@@ -29,9 +30,9 @@ void FInterchangeImportPlugin::StartupModule()
 	UInterchangeManager& InterchangeManager = UInterchangeManager::GetInterchangeManager();
 
 	//Register the translators
-
 	//Textures
 	InterchangeManager.RegisterTranslator(UInterchangeBMPTranslator::StaticClass());
+	InterchangeManager.RegisterTranslator(UInterchangeEXRTranslator::StaticClass());
 	InterchangeManager.RegisterTranslator(UInterchangeJPGTranslator::StaticClass());
 	InterchangeManager.RegisterTranslator(UInterchangePCXTranslator::StaticClass());
 	InterchangeManager.RegisterTranslator(UInterchangePNGTranslator::StaticClass());
diff --git a/Engine/Plugins/Runtime/Interchange/Source/Import/Private/Texture/InterchangeEXRTranslator.cpp b/Engine/Plugins/Runtime/Interchange/Source/Import/Private/Texture/InterchangeEXRTranslator.cpp
new file mode 100644
--- /dev/null
+++ b/Engine/Plugins/Runtime/Interchange/Source/Import/Private/Texture/InterchangeEXRTranslator.cpp
@@ -0,0 +1,147 @@
+// Copyright Epic Games, Inc. All Rights Reserved. 
+#include "Texture/InterchangeEXRTranslator.h"
+
+#include "Engine/Texture.h"
+#include "Engine/Texture2D.h"
+#include "IImageWrapper.h"
+#include "IImageWrapperModule.h"
+#include "LogInterchangeImportPlugin.h"
+#include "Misc/ConfigCacheIni.h"
+#include "Misc/FileHelper.h"
+#include "Misc/Paths.h"
+#include "Misc/ScopedSlowTask.h"
+#include "Modules/ModuleManager.h"
+#include "Nodes/BaseNodeContainer.h"
+#include "TextureNode.h"
+
+
+//////////////////////////////////////////////////////////////////////////
+// PCX helper local function
+
+namespace EXRParserHelper
+{
+} //EXRParserHelper namespace end
+
+bool UInterchangeEXRTranslator::CanImportSourceData(const UInterchangeSourceData* InSourceData) const
+{
+	FString Extension = FPaths::GetExtension(InSourceData->GetFilename());
+	FString EXRExtension = (TEXT("exr;Texture"));
+	return EXRExtension.StartsWith(Extension);
+}
+
+bool UInterchangeEXRTranslator::Translate(const UInterchangeSourceData* SourceData, Interchange::FBaseNodeContainer& BaseNodeContainer) const
+{
+	FString Filename = SourceData->GetFilename();
+	if (!FPaths::FileExists(Filename))
+	{
+		return false;
+	}
+
+	FName DisplayLabel = *FPaths::GetBaseFilename(Filename);
+	Interchange::FNodeUniqueID NodeUID(*Filename);
+	//PCX is creating a UTexture2D
+	TUniquePtr<Interchange::FTextureNode> TextureNode = MakeUnique<Interchange::FTextureNode>(NodeUID, DisplayLabel, UTexture2D::StaticClass());
+	TextureNode->SetPayLoadKey(Filename);
+
+	//Test node change
+ 	//TextureNode->SetCustomLODGroup((uint8)TextureGroup::TEXTUREGROUP_WorldNormalMap);
+ 	//TextureNode->SetCustomCompressionSettings((uint8)TextureCompressionSettings::TC_Normalmap);
+
+	BaseNodeContainer.AddNode(MoveTemp(TextureNode));
+	return true;
+}
+
+const TOptional<Interchange::FImportImage> UInterchangeEXRTranslator::GetPayloadData(const UInterchangeSourceData* SourceData, const FString& PayLoadKey) const
+{
+	if (!SourceData)
+	{
+		UE_LOG(LogInterchangeImportPlugin, Error, TEXT("Failed to import EXR, bad source data."));
+		return TOptional<Interchange::FImportImage>();
+	}
+
+	TArray64<uint8> SourceDataBuffer;
+	FString Filename = SourceData->GetFilename();
+	
+	//Make sure the key fit the filename, The key should always be valid
+	if (!Filename.Equals(PayLoadKey))
+	{
+		UE_LOG(LogInterchangeImportPlugin, Error, TEXT("Failed to import EXR, wrong payload key. [%s]"), *Filename);
+		return TOptional<Interchange::FImportImage>();
+	}
+
+	if (!FPaths::FileExists(Filename))
+	{
+		UE_LOG(LogInterchangeImportPlugin, Error, TEXT("Failed to import EXR, cannot open file. [%s]"), *Filename);
+		return TOptional<Interchange::FImportImage>();
+	}
+
+	if (!FFileHelper::LoadFileToArray(SourceDataBuffer, *Filename))
+	{
+		UE_LOG(LogInterchangeImportPlugin, Error, TEXT("Failed to import EXR, cannot load file content into an array. [%s]"), *Filename);
+		return TOptional<Interchange::FImportImage>();
+	}
+
+	const uint8* Buffer = SourceDataBuffer.GetData();
+	const uint8* BufferEnd = Buffer + SourceDataBuffer.Num();
+
+	bool bAllowNonPowerOfTwo = false;
+	GConfig->GetBool(TEXT("TextureImporter"), TEXT("AllowNonPowerOfTwoTextures"), bAllowNonPowerOfTwo, GEditorIni);
+
+	// Validate it.
+	const int32 Length = BufferEnd - Buffer;
+
+	//
+	// EXR
+	//
+	IImageWrapperModule& ImageWrapperModule = FModuleManager::LoadModuleChecked<IImageWrapperModule>(FName("ImageWrapper"));
+
+	TSharedPtr<IImageWrapper> ExrImageWrapper = ImageWrapperModule.CreateImageWrapper(EImageFormat::EXR);
+	if (!ExrImageWrapper.IsValid() || !ExrImageWrapper->SetCompressed(Buffer, Length))
+	{
+		UE_LOG(LogInterchangeImportPlugin, Error, TEXT("Failed to import EXR, unsupported format. [%s]"), *Filename);
+		return TOptional<Interchange::FImportImage>();
+	}
+
+	int32 Width = ExrImageWrapper->GetWidth();
+	int32 Height = ExrImageWrapper->GetHeight();
+
+	if (!Interchange::FImportImageHelper::IsImportResolutionValid(Width, Height, bAllowNonPowerOfTwo))
+	{
+		UE_LOG(LogInterchangeImportPlugin, Error, TEXT("Failed to import PCX, invalid resolution. Resolution[%d, %d], AllowPowerOfTwo[%s], [%s]"), Width, Height, bAllowNonPowerOfTwo ? TEXT("True") : TEXT("false"), *Filename);
+		return TOptional<Interchange::FImportImage>();
+	}
+
+	// Select the texture's source format
+	ETextureSourceFormat TextureFormat = TSF_Invalid;
+	int32 BitDepth = ExrImageWrapper->GetBitDepth();
+	ERGBFormat Format = ExrImageWrapper->GetFormat();
+
+	if (Format == ERGBFormat::RGBA && BitDepth == 16)
+	{
+		TextureFormat = TSF_RGBA16F;
+		Format = ERGBFormat::BGRA;
+	}
+
+	if (TextureFormat == TSF_Invalid)
+	{
+		UE_LOG(LogInterchangeImportPlugin, Error, TEXT("EXR file [%s] contains data in an unsupported format"), *Filename);
+		return TOptional<Interchange::FImportImage>();
+	}
+
+	Interchange::FImportImage PayloadData;
+	PayloadData.Init2DWithParams(
+		Width,
+		Height,
+		TextureFormat,
+		false
+	);
+	PayloadData.CompressionSettings = TC_HDR;
+
+	if (!ExrImageWrapper->GetRaw(Format, BitDepth, PayloadData.RawData))
+	{
+		UE_LOG(LogInterchangeImportPlugin, Error, TEXT("Failed to decode EXR. [%s]"), *Filename);
+		return TOptional<Interchange::FImportImage>();
+	}
+
+	return PayloadData;
+}
\ No newline at end of file
diff --git a/Engine/Plugins/Runtime/Interchange/Source/Import/Public/Texture/InterchangeEXRTranslator.h b/Engine/Plugins/Runtime/Interchange/Source/Import/Public/Texture/InterchangeEXRTranslator.h
new file mode 100644
--- /dev/null
+++ b/Engine/Plugins/Runtime/Interchange/Source/Import/Public/Texture/InterchangeEXRTranslator.h
@@ -0,0 +1,45 @@
+// Copyright Epic Games, Inc. All Rights Reserved.
+
+#pragma once
+
+#include "CoreMinimal.h"
+#include "UObject/Object.h"
+#include "UObject/ObjectMacros.h"
+#include "InterchangeTranslatorBase.h"
+#include "Texture/InterchangeTexturePayloadData.h"
+#include "Texture/InterchangeTextureTranslator.h"
+
+#include "InterchangeEXRTranslator.generated.h"
+
+UCLASS(BlueprintType)
+class INTERCHANGEIMPORTPLUGIN_API UInterchangeEXRTranslator : public UInterchangeTextureTranslator
+{
+	GENERATED_BODY()
+public:
+	
+	/*
+	 * return true if the translator can translate the specified file.
+	 */
+	virtual bool CanImportSourceData(const UInterchangeSourceData* InSourceData) const override;
+
+	/**
+	 * Translate a source data into a node hold by the specified nodes container.
+	 *
+	 * @param SourceData - The source data containing the data to translate
+	 * @param BaseNodeContainer - The unreal objects descriptions container where to put the translated source data.
+	 * @return true if the translator can translate the source data, false otherwise.
+	 */
+	virtual bool Translate(const UInterchangeSourceData* SourceData, Interchange::FBaseNodeContainer& BaseNodeContainer) const override;
+
+	/**
+	 * Once the translation is done, the import process need a way to retrieve payload data.
+	 * This payload will be use by the factories to create the asset.
+	 *
+	 * @param SourceData - The source data containing the data to translate
+	 * @param PayloadKey - The key to retrieve the a particular payload contain into the specified source data.
+	 * @return a PayloadData containing the import image data. The TOptional will not be set if there is an error.
+	 */
+	virtual const TOptional<Interchange::FImportImage> GetPayloadData(const UInterchangeSourceData* SourceData, const FString& PayLoadKey) const override;
+};
+
+
