Added Interchange translator for .psd textures.

diff --git a/Engine/Plugins/Runtime/Interchange/Source/Import/Private/InterchangeImportPlugin.cpp b/Engine/Plugins/Runtime/Interchange/Source/Import/Private/InterchangeImportPlugin.cpp
--- a/Engine/Plugins/Runtime/Interchange/Source/Import/Private/InterchangeImportPlugin.cpp
+++ b/Engine/Plugins/Runtime/Interchange/Source/Import/Private/InterchangeImportPlugin.cpp
@@ -9,6 +9,7 @@
 #include "Texture/InterchangeJPGTranslator.h"
 #include "Texture/InterchangePCXTranslator.h"
 #include "Texture/InterchangePNGTranslator.h"
+#include "Texture/InterchangePSDTranslator.h"
 #include "Texture/InterchangeTextureFactory.h"
 
 
@@ -36,6 +37,7 @@ void FInterchangeImportPlugin::StartupModule()
 	InterchangeManager.RegisterTranslator(UInterchangeJPGTranslator::StaticClass());
 	InterchangeManager.RegisterTranslator(UInterchangePCXTranslator::StaticClass());
 	InterchangeManager.RegisterTranslator(UInterchangePNGTranslator::StaticClass());
+	InterchangeManager.RegisterTranslator(UInterchangePSDTranslator::StaticClass());
 
 	//Register the factories
 	InterchangeManager.RegisterFactory(UInterchangeTextureFactory::StaticClass());
diff --git a/Engine/Plugins/Runtime/Interchange/Source/Import/Private/Texture/InterchangeJPGTranslator.cpp b/Engine/Plugins/Runtime/Interchange/Source/Import/Private/Texture/InterchangeJPGTranslator.cpp
--- a/Engine/Plugins/Runtime/Interchange/Source/Import/Private/Texture/InterchangeJPGTranslator.cpp
+++ b/Engine/Plugins/Runtime/Interchange/Source/Import/Private/Texture/InterchangeJPGTranslator.cpp
@@ -129,16 +129,16 @@ const TOptional<Interchange::FImportImage> UInterchangeJPGTranslator::GetPayload
 		return TOptional<Interchange::FImportImage>();
 	}
 
-	TOptional<Interchange::FImportImage> PayloadData = Interchange::FImportImage();
+	Interchange::FImportImage PayloadData;
 
-	PayloadData.GetValue().Init2DWithParams(
+	PayloadData.Init2DWithParams(
 		JpegImageWrapper->GetWidth(),
 		JpegImageWrapper->GetHeight(),
 		TextureFormat,
 		BitDepth < 16
 	);
 
-	if (!JpegImageWrapper->GetRaw(Format, BitDepth, PayloadData.GetValue().RawData))
+	if (!JpegImageWrapper->GetRaw(Format, BitDepth, PayloadData.RawData))
 	{
 		UE_LOG(LogInterchangeImportPlugin, Error, TEXT("Failed to decode JPEG. [%s]"), *Filename);
 		return TOptional<Interchange::FImportImage>();
diff --git a/Engine/Plugins/Runtime/Interchange/Source/Import/Private/Texture/InterchangePNGTranslator.cpp b/Engine/Plugins/Runtime/Interchange/Source/Import/Private/Texture/InterchangePNGTranslator.cpp
--- a/Engine/Plugins/Runtime/Interchange/Source/Import/Private/Texture/InterchangePNGTranslator.cpp
+++ b/Engine/Plugins/Runtime/Interchange/Source/Import/Private/Texture/InterchangePNGTranslator.cpp
@@ -308,15 +308,16 @@ const TOptional<Interchange::FImportImage> UInterchangePNGTranslator::GetPayload
 		UE_LOG(LogInterchangeImportPlugin, Error, TEXT("PNG file [%s] contains data in an unsupported format"), *Filename);
 		return TOptional<Interchange::FImportImage>();
 	}
-	TOptional<Interchange::FImportImage> PayloadData = Interchange::FImportImage();
-	PayloadData.GetValue().Init2DWithParams(
+
+	Interchange::FImportImage PayloadData;
+	PayloadData.Init2DWithParams(
 		PngImageWrapper->GetWidth(),
 		PngImageWrapper->GetHeight(),
 		TextureFormat,
 		BitDepth < 16
 	);
 
-	if (PngImageWrapper->GetRaw(Format, BitDepth, PayloadData.GetValue().RawData))
+	if (PngImageWrapper->GetRaw(Format, BitDepth, PayloadData.RawData))
 	{
 		bool bFillPNGZeroAlpha = true;
 		GConfig->GetBool(TEXT("TextureImporter"), TEXT("FillPNGZeroAlpha"), bFillPNGZeroAlpha, GEditorIni);
@@ -324,7 +325,7 @@ const TOptional<Interchange::FImportImage> UInterchangePNGTranslator::GetPayload
 		if (bFillPNGZeroAlpha)
 		{
 			// Replace the pixels with 0.0 alpha with a color value from the nearest neighboring color which has a non-zero alpha
-			PNGParserHelper::FillZeroAlphaPNGData(PayloadData.GetValue().SizeX, PayloadData.GetValue().SizeY, PayloadData.GetValue().Format, PayloadData.GetValue().RawData.GetData());
+			PNGParserHelper::FillZeroAlphaPNGData(PayloadData.SizeX, PayloadData.SizeY, PayloadData.Format, PayloadData.RawData.GetData());
 		}
 	}
 	else
@@ -332,5 +333,6 @@ const TOptional<Interchange::FImportImage> UInterchangePNGTranslator::GetPayload
 		UE_LOG(LogInterchangeImportPlugin, Error, TEXT("Failed to decode PNG. [%s]"), *Filename);
 		return TOptional<Interchange::FImportImage>();
 	}
+
 	return PayloadData;
 }
\ No newline at end of file
diff --git a/Engine/Plugins/Runtime/Interchange/Source/Import/Private/Texture/InterchangePSDTranslator.cpp b/Engine/Plugins/Runtime/Interchange/Source/Import/Private/Texture/InterchangePSDTranslator.cpp
new file mode 100644
--- /dev/null
+++ b/Engine/Plugins/Runtime/Interchange/Source/Import/Private/Texture/InterchangePSDTranslator.cpp
@@ -0,0 +1,489 @@
+// Copyright Epic Games, Inc. All Rights Reserved. 
+#include "Texture/InterchangePSDTranslator.h"
+
+#include "Engine/Texture.h"
+#include "Engine/Texture2D.h"
+#include "IImageWrapper.h"
+#include "IImageWrapperModule.h"
+#include "LogInterchangeImportPlugin.h"
+#include "Misc/ConfigCacheIni.h"
+#include "Misc/FileHelper.h"
+#include "Misc/Paths.h"
+#include "Misc/ScopedSlowTask.h"
+#include "Modules/ModuleManager.h"
+#include "Nodes/BaseNodeContainer.h"
+#include "TextureNode.h"
+
+
+//////////////////////////////////////////////////////////////////////////
+// PSD helper local function
+
+namespace PSDParserHelper
+{
+	#pragma pack(push,1)
+	struct FPSDFileHeader
+	{                                                           
+		int32	Signature;		// 8BPS
+		int16	Version;		// Version
+		int16	nChannels;		// Number of Channels (3=RGB) (4=RGBA)
+		int32	Height;			// Number of Image Rows
+		int32	Width;			// Number of Image Columns
+		int16	Depth;			// Number of Bits per Channel
+		int16	Mode;			// Image Mode (0=Bitmap)(1=Grayscale)(2=Indexed)(3=RGB)(4=CYMK)(7=Multichannel)
+		uint8	Pad[6];			// Padding
+
+		/**
+		 * @return Whether file has a valid signature
+		 */
+		bool IsValid() const
+		{
+			// Fail on bad signature
+			return (Signature == 0x38425053);
+		}
+
+		/**
+		 * @return Whether file has a supported version
+		 */
+		bool IsSupported() const
+		{
+			// Fail on bad version
+			if (Version != 1)
+				return false;
+
+			// Fail on anything other than 1, 3 or 4 channels
+			if ((nChannels != 1) && (nChannels != 3) && (nChannels != 4))
+				return false;
+
+			// Fail on anything other than 8 Bits/channel or 16 Bits/channel  
+			if ((Depth != 8) && (Depth != 16))
+				return false;
+
+			// Fail on anything other than Grayscale and RGB
+			// We can add support for indexed later if needed.
+			if (Mode != 1 && Mode != 3)
+				return false;
+
+			return true;
+		}
+	};
+	#pragma pack(pop)
+
+
+	static bool ReadData(uint8* pOut, const uint8*& pBuffer, FPSDFileHeader& Info)
+	{
+		const uint8* pPlane = nullptr;
+		const uint8* pRowTable = nullptr;
+		int32 iPlane;
+		int16 CompressionType;
+		int32 iPixel;
+		int32 iRow;
+		int32 CompressedBytes;
+		int32 iByte;
+		int32 Count;
+		uint8 Value;
+
+		// Double check to make sure this is a valid request
+		if (!Info.IsValid() || !Info.IsSupported())
+		{
+			return false;
+		}
+
+		const uint8* pCur = pBuffer + sizeof(FPSDFileHeader);
+		int32 NPixels = Info.Width * Info.Height;
+
+		int32 ClutSize = ((int32)pCur[0] << 24) +
+						 ((int32)pCur[1] << 16) +
+						 ((int32)pCur[2] <<  8) +
+						 ((int32)pCur[3] <<  0);
+		pCur += 4 + ClutSize;    
+
+		// Skip Image Resource Section
+		int32 ImageResourceSize = ((int32)pCur[0] << 24) +
+								  ((int32)pCur[1] << 16) +
+								  ((int32)pCur[2] <<  8) +
+								  ((int32)pCur[3] <<  0);
+		pCur += 4 + ImageResourceSize;
+
+		// Skip Layer and Mask Section
+		int32 LayerAndMaskSize =  ((int32)pCur[0] << 24) +
+								  ((int32)pCur[1] << 16) +
+								  ((int32)pCur[2] <<  8) +
+								  ((int32)pCur[3] <<  0);
+		pCur += 4 + LayerAndMaskSize;
+
+		// Determine number of bytes per pixel
+		int32 BytesPerPixel = 3;
+		const int32 BytesPerChannel = (Info.Depth / 8);
+		switch( Info.Mode )
+		{
+		case 1: // 'GrayScale'
+			BytesPerPixel = BytesPerChannel;
+			break;
+		case 2:        
+			BytesPerPixel = 1;        
+			return false;  // until we support indexed...
+			break;
+		case 3: // 'RGBColor'
+			if (Info.nChannels == 3)
+				BytesPerPixel = 3 * BytesPerChannel;
+			else                   
+				BytesPerPixel = 4 * BytesPerChannel;       
+			break;
+		default:
+			return false;
+			break;
+		}
+
+		// Get Compression Type
+		CompressionType = ((int32)pCur[0] <<  8) + ((int32)pCur[1] <<  0);    
+		pCur += 2;
+
+		// Fail on 16 Bits/channel with RLE. This can occur when the file is not saved with 'Maximize Compatibility'. Compression doesn't appear to be standard.
+		if(CompressionType == 1 && Info.Depth == 16)
+		{
+			return false;
+		}
+
+		// If no alpha channel, set alpha to opaque (255 or 65536).
+		if( Info.nChannels != 4)
+		{
+			if(Info.Depth == 8)
+			{
+				const uint32 Channels = 4;
+				const uint32 BufferSize = Info.Width * Info.Height * Channels * sizeof(uint8);
+				FMemory::Memset(pOut, 0xff, BufferSize);
+			}
+			else if(Info.Depth == 16)
+			{
+				const uint32 Channels = 4;
+				const uint32 BufferSize = Info.Width * Info.Height * Channels * sizeof(uint16);
+				FMemory::Memset(pOut, 0xff, BufferSize);
+			}
+		}
+
+		// Uncompressed?
+		if( CompressionType == 0 )
+		{
+			if(Info.Depth == 8)
+			{
+				FColor* Dest = (FColor*)pOut;
+				for(int32 Pixel=0; Pixel < NPixels; Pixel++ )
+				{
+					if (Info.nChannels == 1)
+					{
+						Dest[Pixel].R = pCur[NPixels + Pixel];
+						Dest[Pixel].G = pCur[NPixels + Pixel];
+						Dest[Pixel].B = pCur[NPixels + Pixel];
+					}
+					else
+					{
+						Dest[Pixel].R = pCur[NPixels * 0 + Pixel];
+						Dest[Pixel].G = pCur[NPixels * 1 + Pixel];
+						Dest[Pixel].B = pCur[NPixels * 2 + Pixel];
+						if (Info.nChannels == 4)
+						{
+							Dest[Pixel].A = pCur[NPixels * 3 + Pixel];
+						}
+					}
+				}
+			}
+			else if (Info.Depth == 16)
+			{
+				uint32 SrcOffset = 0;
+
+				if (Info.nChannels == 1)
+				{
+					uint16* Dest = (uint16*)pOut;
+					uint32 ChannelOffset = 0;
+
+					for (int32 Pixel = 0; Pixel < NPixels; Pixel++)
+					{
+						Dest[ChannelOffset + 0] = ((pCur[SrcOffset] << 8) + (pCur[SrcOffset + 1] << 0));
+						Dest[ChannelOffset + 1] = ((pCur[SrcOffset] << 8) + (pCur[SrcOffset + 1] << 0));
+						Dest[ChannelOffset + 2] = ((pCur[SrcOffset] << 8) + (pCur[SrcOffset + 1] << 0));
+
+						//Increment offsets
+						ChannelOffset += 4;
+						SrcOffset += BytesPerChannel;
+					}
+				}
+				else
+				{
+					// Loop through the planes	
+					for (iPlane = 0; iPlane < Info.nChannels; iPlane++)
+					{
+						uint16* Dest = (uint16*)pOut;
+						uint32 ChannelOffset = iPlane;
+
+						for (int32 Pixel = 0; Pixel < NPixels; Pixel++)
+						{
+							Dest[ChannelOffset] = ((pCur[SrcOffset] << 8) + (pCur[SrcOffset + 1] << 0));
+
+							//Increment offsets
+							ChannelOffset += 4;
+							SrcOffset += BytesPerChannel;
+						}
+					}
+				}
+			}
+		}
+		// RLE?
+		else if( CompressionType == 1 )
+		{
+			// Setup RowTable
+			pRowTable = pCur;
+			pCur += Info.nChannels*Info.Height*2;
+
+			FColor* Dest = (FColor*)pOut;
+
+			// Loop through the planes
+			for( iPlane=0 ; iPlane<Info.nChannels ; iPlane++ )
+			{
+				int32 iWritePlane = iPlane;
+				if( iWritePlane > BytesPerPixel-1 ) iWritePlane = BytesPerPixel-1;
+
+				// Loop through the rows
+				for( iRow=0 ; iRow<Info.Height ; iRow++ )
+				{
+					// Load a row
+					CompressedBytes = (pRowTable[(iPlane*Info.Height+iRow)*2  ] << 8) +
+						(pRowTable[(iPlane*Info.Height+iRow)*2+1] << 0);
+
+					// Setup Plane
+					pPlane = pCur;
+					pCur += CompressedBytes;
+
+					// Decompress Row
+					iPixel = 0;
+					iByte = 0;
+					while( (iPixel < Info.Width) && (iByte < CompressedBytes) )
+					{
+						int8 code = (int8)pPlane[iByte++];
+
+						// Is it a repeat?
+						if( code < 0 )
+						{
+							Count = -(int32)code + 1;
+							Value = pPlane[iByte++];
+							while( Count-- > 0 )
+							{
+								int32 idx = (iPixel) + (iRow*Info.Width);
+								if (Info.nChannels == 1)
+								{
+									Dest[idx].R = Value;
+									Dest[idx].G = Value;
+									Dest[idx].B = Value;
+								}
+								else
+								{
+									switch (iWritePlane)
+									{
+									case 0: Dest[idx].R = Value; break;
+									case 1: Dest[idx].G = Value; break;
+									case 2: Dest[idx].B = Value; break;
+									case 3: Dest[idx].A = Value; break;
+									}
+								}
+								iPixel++;
+							}
+						}
+						// Must be a literal then
+						else
+						{
+							Count = (int32)code + 1;
+							while( Count-- > 0 )
+							{
+								Value = pPlane[iByte++];
+								int32 idx = (iPixel) + (iRow*Info.Width);
+
+								if (Info.nChannels == 1)
+								{
+									Dest[idx].R = Value;
+									Dest[idx].G = Value;
+									Dest[idx].B = Value;
+								}
+								else
+								{
+									switch (iWritePlane)
+									{
+									case 0: Dest[idx].R = Value; break;
+									case 1: Dest[idx].G = Value; break;
+									case 2: Dest[idx].B = Value; break;
+									case 3: Dest[idx].A = Value; break;
+									}
+								}
+								iPixel++;
+							}
+						}
+					}
+
+					// Confirm that we decoded the right number of bytes
+					check( iByte  == CompressedBytes );
+					check( iPixel == Info.Width );
+				}
+			}
+		}
+		else
+			return false;
+
+		// Success!
+		return true;
+	}
+
+	static void GetPSDHeader( const uint8* Buffer, FPSDFileHeader& Info )
+	{
+		Info.Signature      =   ((int32)Buffer[ 0] << 24) +
+								((int32)Buffer[ 1] << 16) +
+								((int32)Buffer[ 2] <<  8) +
+								((int32)Buffer[ 3] <<  0);
+
+		Info.Version        =   ((int32)Buffer[ 4] <<  8) +
+								((int32)Buffer[ 5] <<  0);
+
+		Info.nChannels      =   ((int32)Buffer[12] <<  8) +
+								((int32)Buffer[13] <<  0);
+
+		Info.Height         =   ((int32)Buffer[14] << 24) +
+								((int32)Buffer[15] << 16) +
+								((int32)Buffer[16] <<  8) +
+								((int32)Buffer[17] <<  0);
+
+		Info.Width          =   ((int32)Buffer[18] << 24) +
+								((int32)Buffer[19] << 16) +
+								((int32)Buffer[20] <<  8) +
+								((int32)Buffer[21] <<  0);
+
+		Info.Depth          =   ((int32)Buffer[22] <<  8) +
+								((int32)Buffer[23] <<  0);
+
+		Info.Mode           =   ((int32)Buffer[24] <<  8) +
+								((int32)Buffer[25] <<  0);
+	}
+
+} // PSDParserHelper namespace end
+
+
+bool UInterchangePSDTranslator::CanImportSourceData(const UInterchangeSourceData* InSourceData) const
+{
+	FString Extension = FPaths::GetExtension(InSourceData->GetFilename());
+	FString PNGExtension = (TEXT("psd;Texture"));
+	return PNGExtension.StartsWith(Extension);
+}
+
+bool UInterchangePSDTranslator::Translate(const UInterchangeSourceData* SourceData, Interchange::FBaseNodeContainer& BaseNodeContainer) const
+{
+	FString Filename = SourceData->GetFilename();
+	if (!FPaths::FileExists(Filename))
+	{
+		return false;
+	}
+
+	FName DisplayLabel = *FPaths::GetBaseFilename(Filename);
+	Interchange::FNodeUniqueID NodeUID(*Filename);
+	// PSD is creating a UTexture2D
+	TUniquePtr<Interchange::FTextureNode> TextureNode = MakeUnique<Interchange::FTextureNode>(NodeUID, DisplayLabel, UTexture2D::StaticClass());
+	TextureNode->SetPayLoadKey(Filename);
+
+	BaseNodeContainer.AddNode(MoveTemp(TextureNode));
+	return true;
+}
+
+const TOptional<Interchange::FImportImage> UInterchangePSDTranslator::GetPayloadData(const UInterchangeSourceData* SourceData, const FString& PayLoadKey) const
+{
+	if (!SourceData)
+	{
+		UE_LOG(LogInterchangeImportPlugin, Error, TEXT("Failed to import PSD, bad source data."));
+		return TOptional<Interchange::FImportImage>();
+	}
+
+	TArray64<uint8> SourceDataBuffer;
+	FString Filename = SourceData->GetFilename();
+
+	//Make sure the key fit the filename, The key should always be valid
+	if (!Filename.Equals(PayLoadKey))
+	{
+		UE_LOG(LogInterchangeImportPlugin, Error, TEXT("Failed to import PSD, wrong payload key. [%s]"), *Filename);
+		return TOptional<Interchange::FImportImage>();
+	}
+
+	if (!FPaths::FileExists(Filename))
+	{
+		UE_LOG(LogInterchangeImportPlugin, Error, TEXT("Failed to import PSD, cannot open file. [%s]"), *Filename);
+		return TOptional<Interchange::FImportImage>();
+	}
+
+	if (!FFileHelper::LoadFileToArray(SourceDataBuffer, *Filename))
+	{
+		UE_LOG(LogInterchangeImportPlugin, Error, TEXT("Failed to import PSD, cannot load file content into an array. [%s]"), *Filename);
+		return TOptional<Interchange::FImportImage>();
+	}
+
+	const uint8* Buffer = SourceDataBuffer.GetData();
+	const uint8* BufferEnd = Buffer + SourceDataBuffer.Num();
+
+	bool bAllowNonPowerOfTwo = false;
+	GConfig->GetBool(TEXT("TextureImporter"), TEXT("AllowNonPowerOfTwoTextures"), bAllowNonPowerOfTwo, GEditorIni);
+
+	// Validate it.
+	const int32 Length = BufferEnd - Buffer;
+
+	//
+	// PSD File
+	//
+	PSDParserHelper::FPSDFileHeader PSDHeader;
+	if (Length > sizeof(PSDParserHelper::FPSDFileHeader))
+	{
+		PSDParserHelper::GetPSDHeader(Buffer, PSDHeader);
+	}
+
+	if (!PSDHeader.IsValid())
+	{
+		return TOptional<Interchange::FImportImage>();
+	}
+
+	// Check the resolution of the imported texture to ensure validity
+	if (!Interchange::FImportImageHelper::IsImportResolutionValid(PSDHeader.Width, PSDHeader.Height, bAllowNonPowerOfTwo))
+	{
+		return TOptional<Interchange::FImportImage>();
+	}
+
+	if (!PSDHeader.IsSupported())
+	{
+		UE_LOG(LogInterchangeImportPlugin, Error, TEXT("Format of this PSD is not supported. Only Grayscale and RGBColor PSD images are currently supported, in 8-bit or 16-bit."));
+		return TOptional<Interchange::FImportImage>();
+	}
+
+	// Select the texture's source format
+	ETextureSourceFormat TextureFormat = TSF_Invalid;
+	if (PSDHeader.Depth == 8)
+	{
+		TextureFormat = TSF_BGRA8;
+	}
+	else if (PSDHeader.Depth == 16)
+	{
+		TextureFormat = TSF_RGBA16;
+	}
+
+	if (TextureFormat == TSF_Invalid)
+	{
+		UE_LOG(LogInterchangeImportPlugin, Error, TEXT("PSD file contains data in an unsupported format."));
+		return TOptional<Interchange::FImportImage>();
+	}
+
+	Interchange::FImportImage PayloadData;
+
+	// The PSD is supported. Load it up.        
+	PayloadData.Init2DWithOneMip(
+		PSDHeader.Width,
+		PSDHeader.Height,
+		TextureFormat
+	);
+
+	if (!PSDParserHelper::ReadData(PayloadData.RawData.GetData(), Buffer, PSDHeader))
+	{
+		UE_LOG(LogInterchangeImportPlugin, Error, TEXT("Failed to read this PSD"));
+		return TOptional<Interchange::FImportImage>();
+	}
+
+	return PayloadData;
+}
diff --git a/Engine/Plugins/Runtime/Interchange/Source/Import/Public/Texture/InterchangePSDTranslator.h b/Engine/Plugins/Runtime/Interchange/Source/Import/Public/Texture/InterchangePSDTranslator.h
new file mode 100644
--- /dev/null
+++ b/Engine/Plugins/Runtime/Interchange/Source/Import/Public/Texture/InterchangePSDTranslator.h
@@ -0,0 +1,45 @@
+// Copyright Epic Games, Inc. All Rights Reserved.
+
+#pragma once
+
+#include "CoreMinimal.h"
+#include "UObject/Object.h"
+#include "UObject/ObjectMacros.h"
+#include "InterchangeTranslatorBase.h"
+#include "Texture/InterchangeTexturePayloadData.h"
+#include "Texture/InterchangeTextureTranslator.h"
+
+#include "InterchangePSDTranslator.generated.h"
+
+UCLASS(BlueprintType)
+class INTERCHANGEIMPORTPLUGIN_API UInterchangePSDTranslator : public UInterchangeTextureTranslator
+{
+	GENERATED_BODY()
+public:
+
+	/*
+	 * return true if the translator can translate the specified file.
+	 */
+	virtual bool CanImportSourceData(const UInterchangeSourceData* InSourceData) const override;
+
+	/**
+	 * Translate a source data into a node hold by the specified nodes container.
+	 *
+	 * @param SourceData - The source data containing the data to translate
+	 * @param BaseNodeContainer - The unreal objects descriptions container where to put the translated source data.
+	 * @return true if the translator can translate the source data, false otherwise.
+	 */
+	virtual bool Translate(const UInterchangeSourceData* SourceData, Interchange::FBaseNodeContainer& BaseNodeContainer) const override;
+
+	/**
+	 * Once the translation is done, the import process need a way to retrieve payload data.
+	 * This payload will be use by the factories to create the asset.
+	 *
+	 * @param SourceData - The source data containing the data to translate
+	 * @param PayloadKey - The key to retrieve the a particular payload contain into the specified source data.
+	 * @return a PayloadData containing the import image data. The TOptional will not be set if there is an error.
+	 */
+	virtual const TOptional<Interchange::FImportImage> GetPayloadData(const UInterchangeSourceData* SourceData, const FString& PayLoadKey) const override;
+};
+
+
