Support interchange DDS texture format

diff --git a/Engine/Plugins/Runtime/Interchange/Source/Import/Private/InterchangeImportPlugin.cpp b/Engine/Plugins/Runtime/Interchange/Source/Import/Private/InterchangeImportPlugin.cpp
--- a/Engine/Plugins/Runtime/Interchange/Source/Import/Private/InterchangeImportPlugin.cpp
+++ b/Engine/Plugins/Runtime/Interchange/Source/Import/Private/InterchangeImportPlugin.cpp
@@ -5,6 +5,7 @@
 #include "InterchangeManager.h"
 #include "Modules/ModuleManager.h"
 #include "Texture/InterchangeBMPTranslator.h"
+#include "Texture/InterchangeDDSTranslator.h"
 #include "Texture/InterchangeEXRTranslator.h"
 #include "Texture/InterchangeJPGTranslator.h"
 #include "Texture/InterchangePCXTranslator.h"
@@ -33,6 +34,7 @@ void FInterchangeImportPlugin::StartupModule()
 	//Register the translators
 	//Textures
 	InterchangeManager.RegisterTranslator(UInterchangeBMPTranslator::StaticClass());
+	InterchangeManager.RegisterTranslator(UInterchangeDDSTranslator::StaticClass());
 	InterchangeManager.RegisterTranslator(UInterchangeEXRTranslator::StaticClass());
 	InterchangeManager.RegisterTranslator(UInterchangeJPGTranslator::StaticClass());
 	InterchangeManager.RegisterTranslator(UInterchangePCXTranslator::StaticClass());
diff --git a/Engine/Plugins/Runtime/Interchange/Source/Import/Private/Texture/InterchangeDDSTranslator.cpp b/Engine/Plugins/Runtime/Interchange/Source/Import/Private/Texture/InterchangeDDSTranslator.cpp
new file mode 100644
--- /dev/null
+++ b/Engine/Plugins/Runtime/Interchange/Source/Import/Private/Texture/InterchangeDDSTranslator.cpp
@@ -0,0 +1,175 @@
+// Copyright Epic Games, Inc. All Rights Reserved. 
+#include "Texture/InterchangeDDSTranslator.h"
+
+#include "DDSLoader.h"
+#include "Engine/Texture.h"
+#include "Engine/Texture2D.h"
+#include "IImageWrapper.h"
+#include "IImageWrapperModule.h"
+#include "LogInterchangeImportPlugin.h"
+#include "Misc/ConfigCacheIni.h"
+#include "Misc/FileHelper.h"
+#include "Misc/Paths.h"
+#include "Misc/ScopedSlowTask.h"
+#include "Modules/ModuleManager.h"
+#include "Nodes/BaseNodeContainer.h"
+#include "TextureNode.h"
+
+
+//////////////////////////////////////////////////////////////////////////
+// DDS helper local function
+
+namespace DDSParserHelper
+{
+} //DDSParserHelper namespace end
+
+bool UInterchangeDDSTranslator::CanImportSourceData(const UInterchangeSourceData* InSourceData) const
+{
+	FString Extension = FPaths::GetExtension(InSourceData->GetFilename());
+	FString DDSExtension = (TEXT("dds;Texture"));
+	if (!DDSExtension.StartsWith(Extension))
+	{
+		return false;
+	}
+
+	/*
+	 * DDS file can also be a cube map so we have to open the file and see if its a valid 2D texture.
+	 */
+	TArray64<uint8> SourceDataBuffer;
+	FString Filename = InSourceData->GetFilename();
+
+	if (!FPaths::FileExists(Filename))
+	{
+		return false;
+	}
+
+	if (!FFileHelper::LoadFileToArray(SourceDataBuffer, *Filename))
+	{
+		return false;
+	}
+
+	const uint8* Buffer = SourceDataBuffer.GetData();
+	const uint8* BufferEnd = Buffer + SourceDataBuffer.Num();
+	// Validate it.
+	const int32 Length = BufferEnd - Buffer;
+
+
+	FDDSLoadHelper  DDSLoadHelper(Buffer, Length);
+	return DDSLoadHelper.IsValid2DTexture();
+}
+
+bool UInterchangeDDSTranslator::Translate(const UInterchangeSourceData* SourceData, Interchange::FBaseNodeContainer& BaseNodeContainer) const
+{
+	FString Filename = SourceData->GetFilename();
+	if (!FPaths::FileExists(Filename))
+	{
+		return false;
+	}
+
+	FName DisplayLabel = *FPaths::GetBaseFilename(Filename);
+	Interchange::FNodeUniqueID NodeUID(*Filename);
+	//PCX is creating a UTexture2D
+	TUniquePtr<Interchange::FTextureNode> TextureNode = MakeUnique<Interchange::FTextureNode>(NodeUID, DisplayLabel, UTexture2D::StaticClass());
+	TextureNode->SetPayLoadKey(Filename);
+
+	//Test node change
+ 	//TextureNode->SetCustomLODGroup((uint8)TextureGroup::TEXTUREGROUP_WorldNormalMap);
+ 	//TextureNode->SetCustomCompressionSettings((uint8)TextureCompressionSettings::TC_Normalmap);
+
+	BaseNodeContainer.AddNode(MoveTemp(TextureNode));
+	return true;
+}
+
+const TOptional<Interchange::FImportImage> UInterchangeDDSTranslator::GetPayloadData(const UInterchangeSourceData* SourceData, const FString& PayLoadKey) const
+{
+	if (!SourceData)
+	{
+		UE_LOG(LogInterchangeImportPlugin, Error, TEXT("Failed to import DDS, bad source data."));
+		return TOptional<Interchange::FImportImage>();
+	}
+
+	TArray64<uint8> SourceDataBuffer;
+	FString Filename = SourceData->GetFilename();
+	
+	//Make sure the key fit the filename, The key should always be valid
+	if (!Filename.Equals(PayLoadKey))
+	{
+		UE_LOG(LogInterchangeImportPlugin, Error, TEXT("Failed to import DDS, wrong payload key. [%s]"), *Filename);
+		return TOptional<Interchange::FImportImage>();
+	}
+
+	if (!FPaths::FileExists(Filename))
+	{
+		UE_LOG(LogInterchangeImportPlugin, Error, TEXT("Failed to import DDS, cannot open file. [%s]"), *Filename);
+		return TOptional<Interchange::FImportImage>();
+	}
+
+	if (!FFileHelper::LoadFileToArray(SourceDataBuffer, *Filename))
+	{
+		UE_LOG(LogInterchangeImportPlugin, Error, TEXT("Failed to import DDS, cannot load file content into an array. [%s]"), *Filename);
+		return TOptional<Interchange::FImportImage>();
+	}
+
+	const uint8* Buffer = SourceDataBuffer.GetData();
+	const uint8* BufferEnd = Buffer + SourceDataBuffer.Num();
+
+	bool bAllowNonPowerOfTwo = false;
+	GConfig->GetBool(TEXT("TextureImporter"), TEXT("AllowNonPowerOfTwoTextures"), bAllowNonPowerOfTwo, GEditorIni);
+
+	// Validate it.
+	const int32 Length = BufferEnd - Buffer;
+
+
+	//
+	// DDS Texture
+	//
+	FDDSLoadHelper  DDSLoadHelper(Buffer, Length);
+	if (!DDSLoadHelper.IsValid2DTexture())
+	{
+		UE_LOG(LogInterchangeImportPlugin, Error, TEXT("Failed to import DDS, unsupported format. [%s]"), *Filename);
+		return TOptional<Interchange::FImportImage>();
+	}
+	// DDS 2d texture
+	if (!Interchange::FImportImageHelper::IsImportResolutionValid(DDSLoadHelper.DDSHeader->dwWidth, DDSLoadHelper.DDSHeader->dwHeight, bAllowNonPowerOfTwo))
+	{
+		UE_LOG(LogInterchangeImportPlugin, Error, TEXT("Failed to import PCX, invalid resolution. Resolution[%d, %d], AllowPowerOfTwo[%s], [%s]"), DDSLoadHelper.DDSHeader->dwWidth, DDSLoadHelper.DDSHeader->dwHeight, bAllowNonPowerOfTwo ? TEXT("True") : TEXT("false"), *Filename);
+		return TOptional<Interchange::FImportImage>();
+	}
+
+	ETextureSourceFormat SourceFormat = DDSLoadHelper.ComputeSourceFormat();
+
+	// Invalid DDS format
+	if (SourceFormat == TSF_Invalid)
+	{
+		UE_LOG(LogInterchangeImportPlugin, Error, TEXT("DDS file [%s] contains data in an unsupported format"), *Filename);
+		return TOptional<Interchange::FImportImage>();
+	}
+
+	uint32 MipMapCount = DDSLoadHelper.ComputeMipMapCount();
+	if (MipMapCount <= 0)
+	{
+		UE_LOG(LogInterchangeImportPlugin, Error, TEXT("DDS file [%s] do not have any mipmap"), *Filename);
+		return TOptional<Interchange::FImportImage>();
+	}
+
+	Interchange::FImportImage PayloadData;
+	PayloadData.Init2DWithMips(
+		DDSLoadHelper.DDSHeader->dwWidth,
+		DDSLoadHelper.DDSHeader->dwHeight,
+		MipMapCount,
+		SourceFormat,
+		DDSLoadHelper.GetDDSDataPointer()
+	);
+
+	if (MipMapCount > 1)
+	{
+		PayloadData.MipGenSettings = TMGS_LeaveExistingMips;
+	}
+	if (FTextureSource::IsHDR(SourceFormat))
+	{
+		// the loader can suggest a compression setting
+		PayloadData.CompressionSettings = TC_HDR;
+	}
+
+	return PayloadData;
+}
\ No newline at end of file
diff --git a/Engine/Plugins/Runtime/Interchange/Source/Import/Private/Texture/InterchangeEXRTranslator.cpp b/Engine/Plugins/Runtime/Interchange/Source/Import/Private/Texture/InterchangeEXRTranslator.cpp
--- a/Engine/Plugins/Runtime/Interchange/Source/Import/Private/Texture/InterchangeEXRTranslator.cpp
+++ b/Engine/Plugins/Runtime/Interchange/Source/Import/Private/Texture/InterchangeEXRTranslator.cpp
@@ -15,13 +15,6 @@
 #include "TextureNode.h"
 
 
-//////////////////////////////////////////////////////////////////////////
-// PCX helper local function
-
-namespace EXRParserHelper
-{
-} //EXRParserHelper namespace end
-
 bool UInterchangeEXRTranslator::CanImportSourceData(const UInterchangeSourceData* InSourceData) const
 {
 	FString Extension = FPaths::GetExtension(InSourceData->GetFilename());
diff --git a/Engine/Plugins/Runtime/Interchange/Source/Import/Private/Texture/InterchangeTextureFactory.cpp b/Engine/Plugins/Runtime/Interchange/Source/Import/Private/Texture/InterchangeTextureFactory.cpp
--- a/Engine/Plugins/Runtime/Interchange/Source/Import/Private/Texture/InterchangeTextureFactory.cpp
+++ b/Engine/Plugins/Runtime/Interchange/Source/Import/Private/Texture/InterchangeTextureFactory.cpp
@@ -152,6 +152,13 @@ UObject* UInterchangeTextureFactory::CreateAsset(const UInterchangeTextureFactor
 		);
 		Texture2D->CompressionSettings = Image.CompressionSettings;
 		Texture2D->SRGB = Image.SRGB;
+		
+		//If the MipGenSettings was set by the translator, we must apply it before the build
+		if (Image.MipGenSettings.IsSet())
+		{
+			// if the source has mips we keep the mips by default, unless the user changes that
+			Texture2D->MipGenSettings = Image.MipGenSettings.GetValue();
+		}
 
 		//The re-import change only the source data and will not apply any custom attribute to the asset.
 		if(!ExistingAsset)
diff --git a/Engine/Plugins/Runtime/Interchange/Source/Import/Public/Texture/InterchangeDDSTranslator.h b/Engine/Plugins/Runtime/Interchange/Source/Import/Public/Texture/InterchangeDDSTranslator.h
new file mode 100644
--- /dev/null
+++ b/Engine/Plugins/Runtime/Interchange/Source/Import/Public/Texture/InterchangeDDSTranslator.h
@@ -0,0 +1,45 @@
+// Copyright Epic Games, Inc. All Rights Reserved.
+
+#pragma once
+
+#include "CoreMinimal.h"
+#include "UObject/Object.h"
+#include "UObject/ObjectMacros.h"
+#include "InterchangeTranslatorBase.h"
+#include "Texture/InterchangeTexturePayloadData.h"
+#include "Texture/InterchangeTextureTranslator.h"
+
+#include "InterchangeDDSTranslator.generated.h"
+
+UCLASS(BlueprintType)
+class INTERCHANGEIMPORTPLUGIN_API UInterchangeDDSTranslator : public UInterchangeTextureTranslator
+{
+	GENERATED_BODY()
+public:
+	
+	/*
+	 * return true if the translator can translate the specified file.
+	 */
+	virtual bool CanImportSourceData(const UInterchangeSourceData* InSourceData) const override;
+
+	/**
+	 * Translate a source data into a node hold by the specified nodes container.
+	 *
+	 * @param SourceData - The source data containing the data to translate
+	 * @param BaseNodeContainer - The unreal objects descriptions container where to put the translated source data.
+	 * @return true if the translator can translate the source data, false otherwise.
+	 */
+	virtual bool Translate(const UInterchangeSourceData* SourceData, Interchange::FBaseNodeContainer& BaseNodeContainer) const override;
+
+	/**
+	 * Once the translation is done, the import process need a way to retrieve payload data.
+	 * This payload will be use by the factories to create the asset.
+	 *
+	 * @param SourceData - The source data containing the data to translate
+	 * @param PayloadKey - The key to retrieve the a particular payload contain into the specified source data.
+	 * @return a PayloadData containing the import image data. The TOptional will not be set if there is an error.
+	 */
+	virtual const TOptional<Interchange::FImportImage> GetPayloadData(const UInterchangeSourceData* SourceData, const FString& PayLoadKey) const override;
+};
+
+
diff --git a/Engine/Plugins/Runtime/Interchange/Source/Import/Public/Texture/InterchangeTexturePayloadData.h b/Engine/Plugins/Runtime/Interchange/Source/Import/Public/Texture/InterchangeTexturePayloadData.h
--- a/Engine/Plugins/Runtime/Interchange/Source/Import/Public/Texture/InterchangeTexturePayloadData.h
+++ b/Engine/Plugins/Runtime/Interchange/Source/Import/Public/Texture/InterchangeTexturePayloadData.h
@@ -16,6 +16,7 @@ namespace Interchange
 		int32 SizeX = 0;
 		int32 SizeY = 0;
 		bool SRGB = true;
+		TOptional<TextureMipGenSettings> MipGenSettings;
 
 		void Init2DWithParams(int32 InSizeX, int32 InSizeY, ETextureSourceFormat InFormat, bool InSRGB);
 		void Init2DWithOneMip(int32 InSizeX, int32 InSizeY, ETextureSourceFormat InFormat, const void* InData = nullptr);
