Add WITH_EDITOR guards to fix CIS

#rnx

diff --git a/Engine/Plugins/Runtime/Interchange/Source/Import/Private/Texture/InterchangeDDSTranslator.cpp b/Engine/Plugins/Runtime/Interchange/Source/Import/Private/Texture/InterchangeDDSTranslator.cpp
--- a/Engine/Plugins/Runtime/Interchange/Source/Import/Private/Texture/InterchangeDDSTranslator.cpp
+++ b/Engine/Plugins/Runtime/Interchange/Source/Import/Private/Texture/InterchangeDDSTranslator.cpp
@@ -82,6 +82,9 @@ bool UInterchangeDDSTranslator::Translate(const UInterchangeSourceData* SourceDa
 
 const TOptional<Interchange::FImportImage> UInterchangeDDSTranslator::GetPayloadData(const UInterchangeSourceData* SourceData, const FString& PayLoadKey) const
 {
+#if !WITH_EDITOR
+	return {};
+#else
 	if (!SourceData)
 	{
 		UE_LOG(LogInterchangeImportPlugin, Error, TEXT("Failed to import DDS, bad source data."));
@@ -172,4 +175,5 @@ const TOptional<Interchange::FImportImage> UInterchangeDDSTranslator::GetPayload
 	}
 
 	return PayloadData;
-}
\ No newline at end of file
+#endif
+}
diff --git a/Engine/Plugins/Runtime/Interchange/Source/Import/Private/Texture/InterchangePCXTranslator.cpp b/Engine/Plugins/Runtime/Interchange/Source/Import/Private/Texture/InterchangePCXTranslator.cpp
--- a/Engine/Plugins/Runtime/Interchange/Source/Import/Private/Texture/InterchangePCXTranslator.cpp
+++ b/Engine/Plugins/Runtime/Interchange/Source/Import/Private/Texture/InterchangePCXTranslator.cpp
@@ -92,6 +92,9 @@ bool UInterchangePCXTranslator::Translate(const UInterchangeSourceData* SourceDa
 
 const TOptional<Interchange::FImportImage> UInterchangePCXTranslator::GetPayloadData(const UInterchangeSourceData* SourceData, const FString& PayLoadKey) const
 {
+#if !WITH_EDITOR
+	return {};
+#else
 	if (!SourceData)
 	{
 		UE_LOG(LogInterchangeImportPlugin, Error, TEXT("Failed to import PCX, bad source data."));
@@ -254,4 +257,5 @@ const TOptional<Interchange::FImportImage> UInterchangePCXTranslator::GetPayload
 	}
 
 	return PayloadData;
-}
\ No newline at end of file
+#endif
+}
