Ensure PSD header is always valid before accessing it (fix static analysis warning).

diff --git a/Engine/Plugins/Runtime/Interchange/Source/Import/Private/Texture/InterchangePSDTranslator.cpp b/Engine/Plugins/Runtime/Interchange/Source/Import/Private/Texture/InterchangePSDTranslator.cpp
--- a/Engine/Plugins/Runtime/Interchange/Source/Import/Private/Texture/InterchangePSDTranslator.cpp
+++ b/Engine/Plugins/Runtime/Interchange/Source/Import/Private/Texture/InterchangePSDTranslator.cpp
@@ -431,11 +431,13 @@ const TOptional<Interchange::FImportImage> UInterchangePSDTranslator::GetPayload
 	// PSD File
 	//
 	PSDParserHelper::FPSDFileHeader PSDHeader;
-	if (Length > sizeof(PSDParserHelper::FPSDFileHeader))
+	if (Length < sizeof(PSDParserHelper::FPSDFileHeader))
 	{
-		PSDParserHelper::GetPSDHeader(Buffer, PSDHeader);
+		return TOptional<Interchange::FImportImage>();
 	}
 
+	PSDParserHelper::GetPSDHeader(Buffer, PSDHeader);
+
 	if (!PSDHeader.IsValid())
 	{
 		return TOptional<Interchange::FImportImage>();
