Create a InterchangeWorker application to be able to call multiprocess sdk or any other stuff that are not thread safe.

diff --git a/Engine/Plugins/Runtime/Interchange/Interchange.uplugin b/Engine/Plugins/Runtime/Interchange/Interchange.uplugin
--- a/Engine/Plugins/Runtime/Interchange/Interchange.uplugin
+++ b/Engine/Plugins/Runtime/Interchange/Interchange.uplugin
@@ -49,6 +49,26 @@
 				"Linux",
 				"Mac"
 			]
+		},
+		{
+			"Name": "InterchangeDispatcher",
+			"Type": "Runtime",
+			"LoadingPhase": "Default",
+			"WhitelistPlatforms": [
+				"Win64",
+				"Linux",
+				"Mac"
+			]
+		},
+		{
+			"Name": "InterchangeFbxParser",
+			"Type": "Runtime",
+			"LoadingPhase": "Default",
+			"WhitelistPlatforms": [
+				"Win64",
+				"Linux",
+				"Mac"
+			]
 		}
 	]
 }
\ No newline at end of file
diff --git a/Engine/Plugins/Runtime/Interchange/Source/Dispatcher/InterchangeDispatcher.Build.cs b/Engine/Plugins/Runtime/Interchange/Source/Dispatcher/InterchangeDispatcher.Build.cs
new file mode 100644
--- /dev/null
+++ b/Engine/Plugins/Runtime/Interchange/Source/Dispatcher/InterchangeDispatcher.Build.cs
@@ -0,0 +1,28 @@
+// Copyright Epic Games, Inc. All Rights Reserved.
+
+using System.IO;
+
+namespace UnrealBuildTool.Rules
+{
+	public class InterchangeDispatcher : ModuleRules
+	{
+		public InterchangeDispatcher(ReadOnlyTargetRules Target) : base(Target)
+		{
+            PublicDependencyModuleNames.AddRange(
+				new string[]
+				{
+                    "Core",
+					"Sockets",
+				}
+			);
+
+			PrivateDependencyModuleNames.AddRange(
+				new string[]
+				{
+					"Json",
+				}
+			);
+
+        }
+    }
+}
diff --git a/Engine/Plugins/Runtime/Interchange/Source/Dispatcher/Private/InterchangeCommands.cpp b/Engine/Plugins/Runtime/Interchange/Source/Dispatcher/Private/InterchangeCommands.cpp
new file mode 100644
--- /dev/null
+++ b/Engine/Plugins/Runtime/Interchange/Source/Dispatcher/Private/InterchangeCommands.cpp
@@ -0,0 +1,57 @@
+// Copyright Epic Games, Inc. All Rights Reserved.
+
+#include "InterchangeCommands.h"
+
+#include "Serialization/Archive.h"
+#include "Serialization/MemoryWriter.h"
+#include "Serialization/MemoryReader.h"
+
+namespace InterchangeDispatcher
+{
+	TSharedPtr<ICommand> CreateCommand(ECommandId CommandType)
+	{
+		switch (CommandType)
+		{
+			case ECommandId::Ping: return MakeShared<FPingCommand>();
+			case ECommandId::BackPing: return MakeShared<FBackPingCommand>();
+			case ECommandId::RunTask: return MakeShared<FRunTaskCommand>();
+			case ECommandId::NotifyEndTask: return MakeShared<FCompletedTaskCommand>();
+			case ECommandId::Terminate: return MakeShared<FTerminateCommand>();
+		}
+		return nullptr;
+	}
+
+	void SerializeCommand(ICommand& Command, TArray<uint8>& OutBuffer)
+	{
+		FMemoryWriter ArWriter(OutBuffer);
+		uint8 type = static_cast<uint8>(Command.GetType());
+		ArWriter << type;
+		ArWriter << Command;
+	}
+
+	TSharedPtr<ICommand> DeserializeCommand(const TArray<uint8>& InBuffer)
+	{
+		FMemoryReader ArReader(InBuffer);
+		uint8 type = 0;
+		ArReader << type;
+		if (TSharedPtr<ICommand> Command = CreateCommand(static_cast<ECommandId>(type)))
+		{
+			ArReader << *Command;
+			return ArReader.IsError() ? nullptr : Command;
+		}
+		return nullptr;
+	}
+
+	void FRunTaskCommand::SerializeImpl(FArchive& Ar)
+	{
+		Ar << JsonDescription;
+		Ar << JobIndex;
+	}
+
+  	void FCompletedTaskCommand::SerializeImpl(FArchive& Ar)
+	{
+		Ar << ProcessResult;
+		Ar << JSonMessages;
+		Ar << JSonResult;
+	}
+}
diff --git a/Engine/Plugins/Runtime/Interchange/Source/Dispatcher/Private/InterchangeDispatcher.cpp b/Engine/Plugins/Runtime/Interchange/Source/Dispatcher/Private/InterchangeDispatcher.cpp
new file mode 100644
--- /dev/null
+++ b/Engine/Plugins/Runtime/Interchange/Source/Dispatcher/Private/InterchangeDispatcher.cpp
@@ -0,0 +1,174 @@
+// Copyright Epic Games, Inc. All Rights Reserved.
+
+#include "InterchangeDispatcher.h"
+
+#include "HAL/FileManager.h"
+#include "HAL/PlatformProcess.h"
+#include "InterchangeDispatcherConfig.h"
+#include "InterchangeDispatcherLog.h"
+#include "InterchangeDispatcherTask.h"
+#include "Misc/Paths.h"
+#include "Misc/ScopeLock.h"
+
+namespace InterchangeDispatcher
+{
+
+FInterchangeDispatcher::FInterchangeDispatcher(const FString& InResultFolder)
+	: NextTaskIndex(0)
+	, CompletedTaskCount(0)
+	, ResultFolder(InResultFolder)
+{
+}
+
+int32 FInterchangeDispatcher::AddTask(const FString& InJsonDescription)
+{
+	FScopeLock Lock(&TaskPoolCriticalSection);
+	int32 TaskIndex = TaskPool.Emplace(InJsonDescription);
+	TaskPool[TaskIndex].Index = TaskIndex;
+	return TaskIndex;
+}
+
+TOptional<FTask> FInterchangeDispatcher::GetNextTask()
+{
+	FScopeLock Lock(&TaskPoolCriticalSection);
+
+	while (TaskPool.IsValidIndex(NextTaskIndex) && TaskPool[NextTaskIndex].State != ETaskState::UnTreated)
+	{
+		NextTaskIndex++;
+	}
+
+	if (!TaskPool.IsValidIndex(NextTaskIndex))
+	{
+		return TOptional<FTask>();
+	}
+
+	TaskPool[NextTaskIndex].State = ETaskState::Running;
+	return TaskPool[NextTaskIndex++];
+}
+
+void FInterchangeDispatcher::SetTaskState(int32 TaskIndex, ETaskState TaskState, const FString& JsonResult, const FString& JSonMessages)
+{
+	FString JsonDescription;
+	{
+		FScopeLock Lock(&TaskPoolCriticalSection);
+
+		if (!ensure(TaskPool.IsValidIndex(TaskIndex)))
+		{
+			return;
+		}
+
+		FTask& Task = TaskPool[TaskIndex];
+		Task.State = TaskState;
+		Task.JsonResult = JsonResult;
+		Task.JsonMessages = JSonMessages;
+		JsonDescription = Task.JsonDescription;
+
+		if (TaskState == ETaskState::ProcessOk
+		 || TaskState == ETaskState::ProcessFailed)
+		{
+			CompletedTaskCount++;
+		}
+
+		if (TaskState == ETaskState::UnTreated)
+		{
+			NextTaskIndex = TaskIndex;
+		}
+	}
+
+	UE_CLOG(TaskState == ETaskState::ProcessOk, LogInterchangeDispatcher, Verbose, TEXT("Json processed: %s"), *JsonDescription);
+	UE_CLOG(TaskState == ETaskState::UnTreated, LogInterchangeDispatcher, Warning, TEXT("Json resubmitted: %s"), *JsonDescription);
+	UE_CLOG(TaskState == ETaskState::ProcessFailed, LogInterchangeDispatcher, Error, TEXT("Json processing failure: %s"), *JsonDescription);
+}
+
+void FInterchangeDispatcher::GetTaskState(int32 TaskIndex, ETaskState& TaskState, FString& JsonResult, FString& JSonMessages)
+{
+	FScopeLock Lock(&TaskPoolCriticalSection);
+
+	if (!ensure(TaskPool.IsValidIndex(TaskIndex)))
+	{
+		return;
+	}
+
+	FTask& Task = TaskPool[TaskIndex];
+	TaskState = Task.State;
+	JsonResult = Task.JsonResult;
+	JSonMessages = Task.JsonMessages;
+}
+
+void FInterchangeDispatcher::StartProcess()
+{
+	//Start the process
+	SpawnHandler();
+}
+
+void FInterchangeDispatcher::TerminateProcess()
+{
+	//Empty the cache folder
+	if (IFileManager::Get().DirectoryExists(*ResultFolder))
+	{
+		const bool RequireExists = false;
+		//Delete recursively folder's content
+		const bool Tree = true;
+		IFileManager::Get().DeleteDirectory(*ResultFolder, RequireExists, Tree);
+	}
+	//Terminate the process
+	CloseHandler();
+}
+
+void FInterchangeDispatcher::WaitAllTaskToCompleteExecution()
+{
+	if (WorkerHandler.IsValid())
+	{
+		UE_LOG(LogInterchangeDispatcher, Error, TEXT("Cannot execute tasks before starting the process"));
+	}
+
+	bool bLogRestartError = true;
+	while (!IsOver())
+	{
+
+		if (!IsHandlerAlive())
+		{
+			break;
+		}
+
+		FPlatformProcess::Sleep(0.1f);
+	}
+
+	if (!IsOver())
+	{
+		UE_LOG(LogInterchangeDispatcher, Warning,
+			TEXT("Begin local processing. (Multi Process failed to consume all the tasks)\n")
+			TEXT("See workers logs: %sPrograms/InterchangeWorker/Saved/Logs"), *FPaths::ConvertRelativePathToFull(FPaths::EngineDir()));
+	}
+	else
+	{
+		UE_LOG(LogInterchangeDispatcher, Display, TEXT("Multi Process ended and consumed all the tasks"));
+	}
+}
+
+bool FInterchangeDispatcher::IsOver()
+{
+	FScopeLock Lock(&TaskPoolCriticalSection);
+	return CompletedTaskCount == TaskPool.Num();
+}
+
+void FInterchangeDispatcher::SpawnHandler()
+{
+	WorkerHandler = MakeUnique<FInterchangeWorkerHandler>(*this, ResultFolder);
+}
+
+bool FInterchangeDispatcher::IsHandlerAlive()
+{
+	return WorkerHandler.IsValid() && WorkerHandler->IsAlive();
+}
+
+void FInterchangeDispatcher::CloseHandler()
+{
+	if (IsHandlerAlive())
+	{
+		WorkerHandler->Stop();
+	}
+	WorkerHandler.Reset();
+}
+
+} // ns InterchangeDispatcher
\ No newline at end of file
diff --git a/Engine/Plugins/Runtime/Interchange/Source/Dispatcher/Private/InterchangeDispatcherModule.cpp b/Engine/Plugins/Runtime/Interchange/Source/Dispatcher/Private/InterchangeDispatcherModule.cpp
new file mode 100644
--- /dev/null
+++ b/Engine/Plugins/Runtime/Interchange/Source/Dispatcher/Private/InterchangeDispatcherModule.cpp
@@ -0,0 +1,32 @@
+// Copyright Epic Games, Inc. All Rights Reserved.
+#include "InterchangeDispatcherModule.h"
+
+#include "InterchangeDispatcherLog.h"
+
+#include "Modules/ModuleInterface.h"
+#include "Modules/ModuleManager.h"
+
+
+
+#define LOCTEXT_NAMESPACE "FInterchangeDispatcherModule"
+
+DEFINE_LOG_CATEGORY(LogInterchangeDispatcher);
+
+FInterchangeDispatcherModule& FInterchangeDispatcherModule::Get()
+{
+	return FModuleManager::LoadModuleChecked< FInterchangeDispatcherModule >(INTERCHANGEDISPATCHER_MODULE_NAME);
+}
+
+bool FInterchangeDispatcherModule::IsAvailable()
+{
+	return FModuleManager::Get().IsModuleLoaded(INTERCHANGEDISPATCHER_MODULE_NAME);
+}
+
+void FInterchangeDispatcherModule::StartupModule()
+{
+}
+
+IMPLEMENT_MODULE(FInterchangeDispatcherModule, InterchangeDispatcher);
+
+#undef LOCTEXT_NAMESPACE // "InterchangeDispatcherModule"
+
diff --git a/Engine/Plugins/Runtime/Interchange/Source/Dispatcher/Private/InterchangeDispatcherNetworking.cpp b/Engine/Plugins/Runtime/Interchange/Source/Dispatcher/Private/InterchangeDispatcherNetworking.cpp
new file mode 100644
--- /dev/null
+++ b/Engine/Plugins/Runtime/Interchange/Source/Dispatcher/Private/InterchangeDispatcherNetworking.cpp
@@ -0,0 +1,370 @@
+// Copyright Epic Games, Inc. All Rights Reserved.
+
+#include "InterchangeDispatcherNetworking.h"
+
+#include "InterchangeDispatcherLog.h"
+
+#include "CoreMinimal.h"
+#include "SocketSubsystem.h"
+#include "Sockets.h"
+#include "Serialization/MemoryReader.h"
+#include "Serialization/MemoryWriter.h"
+
+
+namespace InterchangeDispatcher
+{
+
+static FTimespan TimespanFromSeconds(double Seconds)
+{
+	return FTimespan(int64(Seconds * ETimespan::TicksPerSecond));
+}
+
+FNetworkNode::~FNetworkNode()
+{
+	CloseSocket(ConnectedSocket);
+}
+
+bool FNetworkNode::SendMessage(const TArray<uint8>& Message, double Timeout_s)
+{
+	if (bWriteError)
+	{
+		UE_LOG(LogInterchangeDispatcher, Error, TEXT("bWriteError flag raised, can't write"));
+		return false;
+	}
+
+	if (ConnectedSocket == nullptr)
+	{
+		UE_LOG(LogInterchangeDispatcher, Error, TEXT("node not connected, can't write"));
+		return false;
+	}
+
+	ESocketConnectionState State = ConnectedSocket->GetConnectionState();
+	if (State == ESocketConnectionState::SCS_ConnectionError)
+	{
+		UE_LOG(LogInterchangeDispatcher, Error, TEXT("connection state error"));
+		return false;
+	}
+
+	bool bCanWrite = ConnectedSocket->Wait(ESocketWaitConditions::WaitForWrite, TimespanFromSeconds(Timeout_s));
+	if (!bCanWrite)
+	{
+		UE_LOG(LogInterchangeDispatcher, Warning, TEXT("can't write on socket"));
+		return false;
+	}
+
+	ConnectedSocket->SetNonBlocking(false);
+
+	auto SendBuffer = [&](const TArray<uint8>& Buffer) -> bool
+	{
+		uint32 BufferSize = Buffer.Num();
+		const uint8* Data = Buffer.GetData();
+
+		uint32 TotalByteSent = 0;
+		bool bSendSucceed = true;
+		while (bSendSucceed && TotalByteSent < BufferSize)
+		{
+			int32 BytesSent = 0;
+			bSendSucceed &= ConnectedSocket->Send(Data + TotalByteSent, BufferSize - TotalByteSent, BytesSent);
+			TotalByteSent += (uint32) BytesSent;
+		}
+		return bSendSucceed && TotalByteSent == BufferSize;
+	};
+
+	// Send the header
+	TArray<uint8> HeaderBuffer;
+	FMemoryWriter ArWriter(HeaderBuffer);
+
+	FMessageHeader Header;
+	Header.ByteSize = Message.Num();
+	ArWriter << Header;
+
+	if (!SendBuffer(HeaderBuffer))
+	{
+		UE_LOG(LogInterchangeDispatcher, Error, TEXT("can't write header on socket"));
+		bWriteError = true;
+		return false;
+	}
+
+	// Send the content
+	if (!SendBuffer(Message))
+	{
+		UE_LOG(LogInterchangeDispatcher, Error, TEXT("can't write content on socket"));
+		bWriteError = true;
+		return false;
+	}
+
+	return true;
+}
+
+bool FNetworkNode::ReceiveMessage(TArray<uint8>& OutMessage, double Timeout_s)
+{
+	if (bReadError)
+	{
+		UE_LOG(LogInterchangeDispatcher, Error, TEXT("ReadError flag raised, can't read"));
+		return false;
+	}
+
+	if (ConnectedSocket == nullptr)
+	{
+		UE_LOG(LogInterchangeDispatcher, Error, TEXT("node not connected, can't read"));
+		return false;
+	}
+
+	ConnectedSocket->SetNonBlocking(true);
+
+	bool bCanRead = ConnectedSocket && ConnectedSocket->Wait(ESocketWaitConditions::WaitForRead, TimespanFromSeconds(Timeout_s));
+	if (!bCanRead)
+	{
+		return false;
+	}
+
+	// Make sure we have a valid header
+	if (IncommingMessage.Header.ByteSize < 0)
+	{
+		static uint32 HeaderByteSize = []()
+		{
+			FMessageHeader Header;
+			TArray<uint8> HeaderBuffer;
+			FMemoryWriter ArWriter(HeaderBuffer);
+			ArWriter << Header;
+			return HeaderBuffer.Num();
+		}();
+
+		uint32 PendingByteCount;
+		if (!ConnectedSocket->HasPendingData(PendingByteCount) || PendingByteCount < HeaderByteSize)
+		{
+			return false;
+		}
+
+		TArray<uint8> HeaderBuffer;
+		HeaderBuffer.AddZeroed(HeaderByteSize);
+		int32 BytesRead = -1;
+		bool bRecvSucceed = ConnectedSocket->Recv(HeaderBuffer.GetData(), HeaderBuffer.Num(), BytesRead);
+		if (!bRecvSucceed || BytesRead != HeaderByteSize)
+		{
+			UE_LOG(LogInterchangeDispatcher, Error, TEXT("Parsed header failed"));
+			bReadError = true;
+			return false;
+		}
+
+		FMemoryReader ArReader(HeaderBuffer);
+		ArReader << IncommingMessage.Header;
+
+		if (IncommingMessage.Header.ByteSize < 0 || IncommingMessage.Header.ByteSize > 1<<20)
+		{
+			UE_LOG(LogInterchangeDispatcher, Error, TEXT("Parsed header failed: bad message size %d"), IncommingMessage.Header.ByteSize);
+			bReadError = true;
+			return false;
+		}
+
+		IncommingMessage.Content.Reserve(IncommingMessage.Header.ByteSize);
+	}
+
+	// fill the message with available data
+	uint32 PendingByteCount;
+	if (IncommingMessage.Header.ByteSize >= 0 && ConnectedSocket->HasPendingData(PendingByteCount))
+	{
+		int32 MissingByteInCurrentMessage = IncommingMessage.Header.ByteSize - IncommingMessage.Content.Num();
+		int32 ReadTarget = int32(PendingByteCount) < MissingByteInCurrentMessage ? PendingByteCount : MissingByteInCurrentMessage;
+		uint8* Destination = IncommingMessage.Content.GetData() + IncommingMessage.Content.Num();
+		IncommingMessage.Content.AddUninitialized(ReadTarget);
+
+		int32 BytesRead = -1;
+		bool bRecvSucceed = ConnectedSocket->Recv(Destination, ReadTarget, BytesRead);
+		if (BytesRead != ReadTarget)
+		{
+			bReadError = true;
+			UE_LOG(LogInterchangeDispatcher, Error, TEXT("Recv issue"));
+			return false;
+		}
+
+		bool bMessageCompleted = BytesRead == MissingByteInCurrentMessage;
+		if (bMessageCompleted)
+		{
+			OutMessage = MoveTemp(IncommingMessage.Content);
+			IncommingMessage = FMessage();
+			return true;
+		}
+	}
+
+	return false;
+}
+
+
+FSocket* FNetworkNode::CreateInternalSocket(const FString& Description)
+{
+	FSocket* Socket = nullptr;
+	ISocketSubsystem* SocketSubsystem = ISocketSubsystem::Get();
+	TSharedRef<FInternetAddr> InternetAddress = SocketSubsystem->CreateInternetAddr();
+	InternetAddress->SetLoopbackAddress();
+
+	FName Protocol = NAME_None;
+	if (ensure(InternetAddress->IsValid()))
+	{
+		Protocol = InternetAddress->GetProtocolType();
+	}
+	FSocket* NewSocket = SocketSubsystem->CreateSocket(NAME_Stream, *Description, Protocol);
+	UE_CLOG(!NewSocket, LogInterchangeDispatcher, Error, TEXT("Socket creation failure"));
+	return NewSocket;
+}
+
+void FNetworkNode::CloseSocket(FSocket*& Socket)
+{
+	if (Socket)
+	{
+		Socket->Close();
+		ISocketSubsystem::Get()->DestroySocket(Socket);
+		Socket = nullptr;
+	}
+}
+
+bool FNetworkNode::IsConnected()
+{
+	return ConnectedSocket && ConnectedSocket->GetConnectionState() == ESocketConnectionState::SCS_Connected;
+}
+
+FNetworkServerNode::FNetworkServerNode()
+{
+	ListeningSocket = CreateInternalSocket(TEXT("Interchange listening socket"));
+
+	if (!ListeningSocket)
+	{
+		ConnectedSocketError = SocketErrorCode::Error_Create;
+		return;
+	}
+
+	ISocketSubsystem* SocketSubsystem = ISocketSubsystem::Get();
+	TSharedRef<FInternetAddr> InternetAddress = SocketSubsystem->CreateInternetAddr();
+	InternetAddress->SetLoopbackAddress();
+
+	int32 BoundPort = SocketSubsystem->BindNextPort(ListeningSocket, *InternetAddress, 1e3, 1);
+	if (BoundPort == 0)
+	{
+		ConnectedSocketError = SocketErrorCode::Error_Bind;
+		UE_LOG(LogInterchangeDispatcher, Error, TEXT("Socket binding failure"));
+		return;
+	}
+	ensure(BoundPort == ListeningSocket->GetPortNo());
+
+	if (!ListeningSocket->Listen(0))
+	{
+		ConnectedSocketError = SocketErrorCode::Error_Listen;
+		UE_LOG(LogInterchangeDispatcher, Error, TEXT("Socket listen failure"));
+		return;
+	}
+}
+
+FNetworkServerNode::~FNetworkServerNode()
+{
+	CloseSocket(ListeningSocket);
+}
+
+int32 FNetworkServerNode::GetListeningPort()
+{
+	return ListeningSocket ? ListeningSocket->GetPortNo() : 0;
+}
+
+bool FNetworkServerNode::Accept(const FString& Description, double Timeout_s)
+{
+	CloseSocket(ConnectedSocket);
+
+	// wait until a connection occurs
+	ListeningSocket->SetNonBlocking();
+	double AcceptLimit_s = FPlatformTime::Seconds() + Timeout_s;
+	do
+	{
+		ConnectedSocket = ListeningSocket->Accept(Description);
+		if (ConnectedSocket)
+		{
+			return true;
+		}
+		FPlatformProcess::Sleep(0.1);
+	}
+	while(FPlatformTime::Seconds() < AcceptLimit_s);
+
+	return false;
+}
+
+bool FNetworkClientNode::Connect(const FString& Description, int32 ServerPort, double Timeout_s)
+{
+	CloseSocket(ConnectedSocket);
+	ConnectedSocket = CreateInternalSocket(Description);
+
+	if (!ConnectedSocket)
+	{
+		ConnectedSocketError = SocketErrorCode::Error_Create;
+		return false;
+	}
+
+	ISocketSubsystem* SocketSubsystem = ISocketSubsystem::Get();
+	TSharedRef<FInternetAddr> InternetAddress = SocketSubsystem->CreateInternetAddr();
+	InternetAddress->SetLoopbackAddress();
+	InternetAddress->SetPort(ServerPort);
+
+	double ConnectTimeout_s = FPlatformTime::Seconds() + Timeout_s;
+	ConnectedSocket->SetNonBlocking(true);
+	do
+	{
+		if (ConnectedSocket->Connect(*InternetAddress))
+		{
+			UE_LOG(LogInterchangeDispatcher, Verbose, TEXT("Client Node is connected"));
+			return true;
+		}
+	}
+	while(FPlatformTime::Seconds() < ConnectTimeout_s);
+
+	UE_LOG(LogInterchangeDispatcher, Error, TEXT("Client socket failed to connect"));
+	CloseSocket(ConnectedSocket);
+	return false;
+}
+
+void FCommandQueue::SetNetworkInterface(FNetworkNode* InNetworkInterface)
+{
+	NetworkInterface = InNetworkInterface;
+}
+
+TSharedPtr<ICommand> FCommandQueue::GetNextCommand(double Timeout_s)
+{
+	// consume all available commands
+	while (Poll(Timeout_s))
+	{
+		Timeout_s = 0;
+	}
+
+	TSharedPtr<ICommand> OutCommand;
+	return InCommands.Dequeue(OutCommand) ? OutCommand : nullptr;
+}
+
+bool FCommandQueue::SendCommand(ICommand& Commmand, double Timeout_s)
+{
+	TArray<uint8> CommandBuffer;
+	SerializeCommand(Commmand, CommandBuffer);
+	return NetworkInterface && NetworkInterface->SendMessage(CommandBuffer, Timeout_s);
+}
+
+void FCommandQueue::Disconnect(double Timeout_s)
+{
+	// consume all available commands before closing the connection
+	while (Poll(Timeout_s))
+	{
+		Timeout_s = 0;
+	}
+
+	NetworkInterface = nullptr;
+}
+
+bool FCommandQueue::Poll(double Timeout_s)
+{
+	TArray<uint8> CommandBuffer;
+	if (NetworkInterface && NetworkInterface->ReceiveMessage(CommandBuffer, Timeout_s))
+	{
+		if (TSharedPtr<ICommand> Command = DeserializeCommand(CommandBuffer))
+		{
+			InCommands.Enqueue(Command);
+			return true;
+		}
+	}
+	return false;
+}
+
+} // ns InterchangeDispatcher
diff --git a/Engine/Plugins/Runtime/Interchange/Source/Dispatcher/Private/InterchangeDispatcherTask.cpp b/Engine/Plugins/Runtime/Interchange/Source/Dispatcher/Private/InterchangeDispatcherTask.cpp
new file mode 100644
--- /dev/null
+++ b/Engine/Plugins/Runtime/Interchange/Source/Dispatcher/Private/InterchangeDispatcherTask.cpp
@@ -0,0 +1,113 @@
+// Copyright Epic Games, Inc. All Rights Reserved.
+
+#pragma once
+
+#include "InterchangeDispatcherTask.h"
+
+#include "CoreMinimal.h"
+#include "Dom/JsonValue.h"
+#include "Dom/JsonObject.h"
+#include "Serialization/JsonReader.h"
+#include "Serialization/JsonWriter.h"
+#include "Serialization/JsonSerializer.h"
+
+namespace InterchangeDispatcher
+{
+	FString FJsonLoadSourceCmd::ToJson() const
+	{
+		//Code should not do a ToJson if the data was not set before
+		ensure(bIsDataInitialize);
+
+		TSharedPtr<FJsonObject> CmdObject = MakeShared<FJsonObject>();
+		TSharedPtr<FJsonObject> ActionDataObject = MakeShared<FJsonObject>();
+		//CmdObject
+		CmdObject->SetStringField(GetCommandIDKey(), GetAction());
+		CmdObject->SetStringField(GetTranslatorIDKey(), GetTranslatorID());
+		ActionDataObject->SetStringField(GetSourceFilenameKey(), GetSourceFilename());
+		CmdObject->SetObjectField(GetCommandDataKey(), ActionDataObject);
+
+		FString LoadSourceCmd;
+		TSharedRef< TJsonWriter< TCHAR, TPrettyJsonPrintPolicy<TCHAR> > > JsonWriter = TJsonWriterFactory< TCHAR, TPrettyJsonPrintPolicy<TCHAR> >::Create(&LoadSourceCmd);
+		if (!FJsonSerializer::Serialize(CmdObject.ToSharedRef(), JsonWriter))
+		{
+			//Error creating the json cmd string 
+			return FString();
+		}
+		return LoadSourceCmd;
+	}
+
+	bool FJsonLoadSourceCmd::FromJson(const FString& JsonString)
+	{
+		TSharedRef<TJsonReader<TCHAR>> Reader = FJsonStringReader::Create(JsonString);
+
+		TSharedPtr<FJsonObject> CmdObject;
+		if (!FJsonSerializer::Deserialize(Reader, CmdObject) || !CmdObject.IsValid())
+		{
+			//Cannot read the json file
+			return false;
+		}
+		FString JsonActionValue;
+		if (!CmdObject->TryGetStringField(GetCommandIDKey(), JsonActionValue))
+		{
+			//The json cmd id key is missing
+			return false;
+		}
+
+		if (JsonActionValue != GetAction())
+		{
+			//This json do not represent a load command
+			return false;
+		}
+
+		//Read the json
+		if (!CmdObject->TryGetStringField(GetTranslatorIDKey(), TranslatorID))
+		{
+			//Missing Load command translator ID
+			return false;
+		}
+		const TSharedPtr<FJsonObject>* ActionDataObject = nullptr;
+		if (!CmdObject->TryGetObjectField(GetCommandDataKey(), ActionDataObject))
+		{
+			//Missing Load Action data object
+			return false;
+		}
+		if (!((*ActionDataObject)->TryGetStringField(GetSourceFilenameKey(), SourceFilename)))
+		{
+			return false;
+		}
+
+		//Since we filled the data from the json file, set the data has been initialize.
+		bIsDataInitialize = true;
+		return true;
+	}
+
+	FString FJsonLoadSourceCmd::JsonResultParser::ToJson() const
+	{
+		TSharedPtr<FJsonObject> ResultObject = MakeShared<FJsonObject>();
+		//CmdObject
+		ResultObject->SetStringField(GetResultFilenameKey(), GetResultFilename());
+
+		FString JsonResult;
+		TSharedRef< TJsonWriter< TCHAR, TPrettyJsonPrintPolicy<TCHAR> > > JsonWriter = TJsonWriterFactory< TCHAR, TPrettyJsonPrintPolicy<TCHAR> >::Create(&JsonResult);
+		if (!FJsonSerializer::Serialize(ResultObject.ToSharedRef(), JsonWriter))
+		{
+			//Error creating the json cmd string 
+			return FString();
+		}
+		return JsonResult;
+	}
+
+	bool FJsonLoadSourceCmd::JsonResultParser::FromJson(const FString& JsonString)
+	{
+		TSharedRef<TJsonReader<TCHAR>> Reader = FJsonStringReader::Create(JsonString);
+
+		TSharedPtr<FJsonObject> ResultObject;
+		if (!FJsonSerializer::Deserialize(Reader, ResultObject) || !ResultObject.IsValid())
+		{
+			//Cannot read the json file
+			return false;
+		}
+		return ResultObject->TryGetStringField(GetResultFilenameKey(), ResultFilename);
+	}
+
+} // NS InterchangeDispatcher
diff --git a/Engine/Plugins/Runtime/Interchange/Source/Dispatcher/Private/InterchangeWorkerHandler.cpp b/Engine/Plugins/Runtime/Interchange/Source/Dispatcher/Private/InterchangeWorkerHandler.cpp
new file mode 100644
--- /dev/null
+++ b/Engine/Plugins/Runtime/Interchange/Source/Dispatcher/Private/InterchangeWorkerHandler.cpp
@@ -0,0 +1,344 @@
+// Copyright Epic Games, Inc. All Rights Reserved.
+#include "InterchangeWorkerHandler.h"
+
+#include "InterchangeCommands.h"
+#include "InterchangeDispatcher.h"
+#include "InterchangeDispatcherConfig.h"
+#include "InterchangeDispatcherLog.h"
+
+#include "HAL/FileManager.h"
+#include "HAL/PlatformProcess.h"
+#include "Misc/Paths.h"
+#include "Sockets.h"
+#include "SocketSubsystem.h"
+
+namespace InterchangeDispatcher
+{
+
+static FString GetWorkerExecutablePath()
+{
+	static FString ProcessorPath = [&]()
+	{
+		FString Path = FPaths::Combine(FPaths::EnginePluginsDir(), TEXT("Runtime/Interchange/Binaries"));
+
+#if PLATFORM_MAC
+		Path = FPaths::Combine(Path, TEXT("Mac/InterchangeWorker"));
+#elif PLATFORM_LINUX
+		Path = FPaths::Combine(Path, TEXT("Linux/InterchangeWorker"));
+#elif PLATFORM_WINDOWS
+	#if UE_BUILD_DEBUG
+		Path = FPaths::Combine(Path, TEXT("Win64/InterchangeWorker-Win64-Debug.exe"));
+	#else
+		Path = FPaths::Combine(Path, TEXT("Win64/InterchangeWorker.exe"));
+	#endif
+#endif
+		if (!FPaths::FileExists(Path))
+		{
+			UE_LOG(LogInterchangeDispatcher, Warning, TEXT("InterchangeWorker executable not found. Expected location: %s"), *FPaths::ConvertRelativePathToFull(ProcessorPath));
+		}
+		return Path;
+	}();
+
+	return ProcessorPath;
+}
+
+FInterchangeWorkerHandler::FInterchangeWorkerHandler(FInterchangeDispatcher& InDispatcher, FString& InResultFolder)
+	: Dispatcher(InDispatcher)
+	, WorkerState(EWorkerState::Uninitialized)
+	, ErrorState(EWorkerErrorState::Ok)
+	, ResultFolder(InResultFolder)
+	, bShouldTerminate(false)
+{
+	//Use a randomGuid to generate unique name
+	int32 UniqueID = (FPlatformTime::Cycles64() & 0x00000000EFFFFFFF);
+	ThreadName = FString(TEXT("InterchangeWorkerHandler_")) + FString::FromInt(UniqueID);
+	IOThread = FThread(*ThreadName, [this]() { Run(); } );
+}
+
+FInterchangeWorkerHandler::~FInterchangeWorkerHandler()
+{
+	IOThread.Join();
+}
+
+void FInterchangeWorkerHandler::StartWorkerProcess()
+{
+	uint32 WorkerProcessId = 0;
+
+	ensure(ErrorState == EWorkerErrorState::Ok);
+	FString ProcessorPath = GetWorkerExecutablePath();
+	if (FPaths::FileExists(ProcessorPath))
+	{
+		int32 ListenPort = NetworkInterface.GetListeningPort();
+		if (ListenPort == 0)
+		{
+			ErrorState = EWorkerErrorState::ConnectionFailed_NotBound;
+			return;
+		}
+
+		FString CommandToProcess;
+
+		// Manually set worker BaseDir of worker process (as automatic deduction is broken in our case,
+		// probably due to the abuse of ExeBinariesSubFolder in InterchangeWorker.Target.cs)
+		// This fixes paths of logs generated by the worker.
+		FString BaseDir = FPlatformProcess::BaseDir();
+		CommandToProcess += TEXT(" -basedir=\"") + BaseDir + TEXT("\"");
+		CommandToProcess += TEXT(" -ServerPID ") + FString::FromInt(FPlatformProcess::GetCurrentProcessId());
+		CommandToProcess += TEXT(" -ServerPort ") + FString::FromInt(ListenPort);
+		CommandToProcess += TEXT(" -InterchangeDispatcherVersion \"") + DispatcherCommandVersion::ToString() + TEXT('"');
+		CommandToProcess += TEXT(" -ResultFolder \"") + ResultFolder + TEXT("\"");
+		UE_LOG(LogInterchangeDispatcher, Verbose, TEXT("CommandToProcess: %s"), *CommandToProcess);
+		WorkerHandle = FPlatformProcess::CreateProc(*ProcessorPath, *CommandToProcess, true, false, false, &WorkerProcessId, 0, nullptr, nullptr);
+	}
+
+	if (!WorkerHandle.IsValid() || !FPlatformProcess::IsProcRunning(WorkerHandle) || !FPlatformProcess::IsApplicationRunning(WorkerProcessId))
+	{
+		ErrorState = EWorkerErrorState::WorkerProcess_CantCreate;
+		return;
+	}
+}
+
+void FInterchangeWorkerHandler::ValidateConnection()
+{
+	if (!NetworkInterface.IsValid())
+	{
+		UE_LOG(LogInterchangeDispatcher, Error, TEXT("NetworkInterface lost"));
+		WorkerState = EWorkerState::Closing;
+		ErrorState = EWorkerErrorState::ConnectionLost;
+	}
+	else if (WorkerHandle.IsValid() && !FPlatformProcess::IsProcRunning(WorkerHandle))
+	{
+		UE_LOG(LogInterchangeDispatcher, Error, TEXT("Worker lost"));
+		WorkerState = EWorkerState::Closing;
+		ErrorState = EWorkerErrorState::WorkerProcess_Lost;
+	}
+}
+
+void FInterchangeWorkerHandler::Run()
+{
+	WorkerState = EWorkerState::Uninitialized;
+	RunInternal();
+	WorkerState = EWorkerState::Terminated;
+	UE_CLOG(ErrorState != EWorkerErrorState::Ok, LogInterchangeDispatcher, Warning, TEXT("Handler ended with error: %s"), EWorkerErrorStateAsString(ErrorState));
+}
+
+void FInterchangeWorkerHandler::RunInternal()
+{
+	while (IsAlive())
+	{
+		switch (WorkerState)
+		{
+			case EWorkerState::Uninitialized:
+			{
+				ErrorState = EWorkerErrorState::Ok;
+
+				StartWorkerProcess();
+
+				if (ErrorState != EWorkerErrorState::Ok)
+				{
+					WorkerState = EWorkerState::Terminated;
+					break;
+				}
+
+				// The Accept() call on the server blocks until a connection is initiated from a client
+				static const FString SocketDescription = TEXT("InterchangeWorkerHandler");
+				if (!NetworkInterface.Accept(SocketDescription, Config::AcceptTimeout_s))
+				{
+					ErrorState = EWorkerErrorState::ConnectionFailed_NoClient;
+				}
+				else
+				{
+					CommandIO.SetNetworkInterface(&NetworkInterface);
+				}
+
+				if (ErrorState != EWorkerErrorState::Ok)
+				{
+					WorkerState = EWorkerState::Closing;
+					break;
+				}
+
+				WorkerState = EWorkerState::Idle;
+				break;
+			}
+
+			case EWorkerState::Idle:
+			{
+				// Fetch a new task
+				ensureMsgf(CurrentTask.IsSet() == false, TEXT("We should not have a current task when fetching a new one"));
+				CurrentTask = Dispatcher.GetNextTask();
+
+				if (CurrentTask.IsSet())
+				{
+					FRunTaskCommand NewTask(CurrentTask.GetValue());
+
+					if (CommandIO.SendCommand(NewTask, Config::SendCommandTimeout_s))
+					{
+						UE_LOG(LogInterchangeDispatcher, Verbose, TEXT("New task command sent"));
+						WorkerState = EWorkerState::Processing;
+					}
+					else
+					{
+						// Signal that the Task was not processed
+						Dispatcher.SetTaskState(CurrentTask->Index, ETaskState::UnTreated, FString(), FString());
+
+						UE_LOG(LogInterchangeDispatcher, Error, TEXT("New task command issue"));
+						WorkerState = EWorkerState::Closing;
+						ErrorState = EWorkerErrorState::ConnectionLost_SendFailed;
+					}
+				}
+				else if (bShouldTerminate)
+				{
+					UE_LOG(LogInterchangeDispatcher, Verbose, TEXT("Exit loop gracefully"));
+					WorkerState = EWorkerState::Closing;
+				}
+				else
+				{
+					ValidateConnection();
+
+					// consume
+					if (TSharedPtr<ICommand> Command = CommandIO.GetNextCommand(Config::IdleLoopDelay))
+					{
+						ProcessCommand(*Command);
+					}
+				}
+
+				break;
+			}
+
+			case EWorkerState::Processing:
+			{
+				if (TSharedPtr<ICommand> Command = CommandIO.GetNextCommand(Config::ProcessingLoopDelay))
+				{
+					ProcessCommand(*Command);
+
+					bool bProcessingOver = CurrentTask.IsSet() == false;
+					if (bProcessingOver)
+					{
+						WorkerState = bShouldTerminate ? EWorkerState::Closing : EWorkerState::Idle;
+					}
+				}
+				else
+				{
+					ValidateConnection();
+					if (ErrorState == EWorkerErrorState::WorkerProcess_Lost)
+					{
+						if (CurrentTask.IsSet())
+						{
+							Dispatcher.SetTaskState(CurrentTask->Index, ETaskState::ProcessFailed, FString(), FString());
+							CurrentTask.Reset();
+						}
+
+						WorkerState = EWorkerState::Closing;
+					}
+				}
+				break;
+			}
+
+			case EWorkerState::Closing:
+			{
+				// try to close the process gracefully
+				if (WorkerHandle.IsValid())
+				{
+					bool CloseByCommand = Config::CloseProcessByCommand && CommandIO.IsValid() && FPlatformProcess::IsProcRunning(WorkerHandle);
+
+					bool bClosed = false;
+					if (CloseByCommand)
+					{
+						FTerminateCommand Terminate;
+						CommandIO.SendCommand(Terminate, 0);
+
+						for (int32 i = 0; i < int32(10. * Config::TerminateTimeout_s); ++i)
+						{
+							if (!FPlatformProcess::IsProcRunning(WorkerHandle))
+							{
+								bClosed = true;
+								break;
+							}
+							FPlatformProcess::Sleep(0.1);
+						}
+					}
+
+					if (!bClosed)
+					{
+						FPlatformProcess::TerminateProc(WorkerHandle, true);
+					}
+				}
+
+				// Process commands still in input queue
+				CommandIO.Disconnect(0);
+				while (TSharedPtr<ICommand> Command = CommandIO.GetNextCommand(0))
+				{
+					ProcessCommand(*Command);
+				}
+
+				WorkerState = EWorkerState::Terminated;
+				break;
+			}
+
+			default:
+			{
+				ensureMsgf(false, TEXT("missing case handling"));
+			}
+		}
+	}
+}
+
+void FInterchangeWorkerHandler::Stop()
+{
+	bShouldTerminate = true;
+}
+
+bool FInterchangeWorkerHandler::IsAlive() const
+{
+	return WorkerState != EWorkerState::Terminated;
+}
+
+void FInterchangeWorkerHandler::ProcessCommand(ICommand& Command)
+{
+	switch (Command.GetType())
+	{
+		case ECommandId::Ping:
+			ProcessCommand(StaticCast<FPingCommand&>(Command));
+			break;
+
+		case ECommandId::NotifyEndTask:
+			ProcessCommand(StaticCast<FCompletedTaskCommand&>(Command));
+			break;
+
+		default:
+			break;
+	}
+}
+
+void FInterchangeWorkerHandler::ProcessCommand(FPingCommand& PingCommand)
+{
+	InterchangeDispatcher::FBackPingCommand BackPing;
+	CommandIO.SendCommand(BackPing, 0);
+}
+
+void FInterchangeWorkerHandler::ProcessCommand(FCompletedTaskCommand& CompletedTaskCommand)
+{
+	if (!CurrentTask.IsSet())
+	{
+		return;
+	}
+
+	Dispatcher.SetTaskState(CurrentTask->Index, CompletedTaskCommand.ProcessResult, CompletedTaskCommand.JSonResult, CompletedTaskCommand.JSonMessages);
+	CurrentTask.Reset();
+}
+
+const TCHAR* FInterchangeWorkerHandler::EWorkerErrorStateAsString(EWorkerErrorState Error)
+{
+	switch (Error)
+	{
+		case EWorkerErrorState::Ok: return TEXT("Ok");
+		case EWorkerErrorState::ConnectionFailed_NotBound: return TEXT("Connection failed (socket not bound)");
+		case EWorkerErrorState::ConnectionFailed_NoClient: return TEXT("Connection failed (No connection from client)");
+		case EWorkerErrorState::ConnectionLost: return TEXT("Connection lost");
+		case EWorkerErrorState::ConnectionLost_SendFailed: return TEXT("Connection lost (send failed)");
+		case EWorkerErrorState::WorkerProcess_CantCreate: return TEXT("Worker process issue (cannot create worker process)");
+		case EWorkerErrorState::WorkerProcess_Lost: return TEXT("Worker process issue (worker lost)");
+		default: return TEXT("unknown");
+	}
+}
+
+} // ns InterchangeDispatcher
diff --git a/Engine/Plugins/Runtime/Interchange/Source/Dispatcher/Public/InterchangeCommands.h b/Engine/Plugins/Runtime/Interchange/Source/Dispatcher/Public/InterchangeCommands.h
new file mode 100644
--- /dev/null
+++ b/Engine/Plugins/Runtime/Interchange/Source/Dispatcher/Public/InterchangeCommands.h
@@ -0,0 +1,147 @@
+// Copyright Epic Games, Inc. All Rights Reserved.
+
+#pragma once
+
+#include "CoreMinimal.h"
+#include "InterchangeDispatcherTask.h"
+#include "Serialization/Archive.h"
+
+class FArchive;
+
+namespace InterchangeDispatcher
+{
+
+class DispatcherCommandVersion
+{
+public:
+	//Major version should be updated when there is an existing API that has been change.
+	static int32 GetMajor() {return 1;}
+	//Minor version must be updated when there is an addition to the API.
+	static int32 GetMinor() { return 0; }
+	//Patch version should be update if there is some bug fixes in the private code.
+	static int32 GetPatch() { return 0; }
+
+	/** Return the version in string format "Major.Minor.Patch" */
+	static FString ToString()
+	{
+		static FString Version = FString::FromInt(GetMajor()) + TEXT(".") + FString::FromInt(GetMinor()) + TEXT(".") + FString::FromInt(GetPatch());
+		return Version;
+	}
+
+	static bool FromString(const FString& VersionStr, int32& OutMajor, int32& OutMinor, int32& OutPatch)
+	{
+		TArray<FString> Tokens;
+		VersionStr.ParseIntoArray(Tokens, TEXT("."));
+		if (Tokens.Num() != 3)
+		{
+			return false;
+		}
+		OutMajor = FCString::Atoi(*(Tokens[0]));
+		OutMinor = FCString::Atoi(*(Tokens[1]));
+		OutPatch = FCString::Atoi(*(Tokens[2]));
+		return true;
+	}
+
+	/**
+	 * We consider having the same major and minor version will make the API fully compatible.
+	 * Patch is only a hint in case you need a particular fixes
+	 * @param Major - The Major version you want to test against the compile Major version
+	 * @param Minor - The Minor version you want to test against the compile Minor version
+	 * @param Patch - In case we need it later, not use currently
+	 *
+	 * @return true if Major and Minor version equal the compile one, InPatch is only there in case we need it in a future release
+	 */
+	static bool IsAPICompatible(int32 Major, int32 Minor, int32 Patch)
+	{
+		return (GetMajor() == Major && GetMinor() == Minor);
+	}
+};
+
+enum class ECommandId : uint8
+{
+	Invalid,
+	Ping,
+	BackPing,
+	RunTask,
+	NotifyEndTask,
+	Terminate,
+	Last
+};
+
+class ICommand
+{
+public:
+	virtual ~ICommand() = default;
+
+	virtual ECommandId GetType() const = 0;
+
+	friend void operator<<(FArchive& Ar, ICommand& C) { C.SerializeImpl(Ar); }
+
+protected:
+	virtual void SerializeImpl(FArchive&) {}
+};
+	
+
+
+// Create a new command from its type
+TSharedPtr<ICommand> CreateCommand(ECommandId CommandType);
+
+// Converts a command into a byte buffer
+void SerializeCommand(ICommand& Command, TArray<uint8>& OutBuffer);
+
+// Converts byte buffer back into a Command
+// returns nullptr in case of error
+TSharedPtr<ICommand> DeserializeCommand(const TArray<uint8>& InBuffer);
+
+
+
+class FTerminateCommand : public ICommand
+{
+public:
+	virtual ECommandId GetType() const override { return ECommandId::Terminate; }
+};
+
+
+class FPingCommand : public ICommand
+{
+public:
+	virtual ECommandId GetType() const override { return ECommandId::Ping; }
+};
+
+
+class FBackPingCommand : public ICommand
+{
+public:
+	virtual ECommandId GetType() const override { return ECommandId::BackPing; }
+};
+
+class FRunTaskCommand : public ICommand
+{
+public:
+	FRunTaskCommand() = default;
+	FRunTaskCommand(const FTask& Task) : JsonDescription(Task.JsonDescription), JobIndex(Task.Index) {}
+	virtual ECommandId GetType() const override { return ECommandId::RunTask; }
+
+protected:
+	virtual void SerializeImpl(FArchive&) override;
+
+public:
+	FString JsonDescription;
+	int32 JobIndex = -1;
+};
+
+class FCompletedTaskCommand : public ICommand
+{
+public:
+	virtual ECommandId GetType() const override { return ECommandId::NotifyEndTask; }
+
+protected:
+	virtual void SerializeImpl(FArchive&) override;
+
+public:
+	ETaskState ProcessResult = ETaskState::Unknown;
+	FString JSonResult;
+	FString JSonMessages;
+};
+
+} // ns InterchangeDispatcher
diff --git a/Engine/Plugins/Runtime/Interchange/Source/Dispatcher/Public/InterchangeDispatcher.h b/Engine/Plugins/Runtime/Interchange/Source/Dispatcher/Public/InterchangeDispatcher.h
new file mode 100644
--- /dev/null
+++ b/Engine/Plugins/Runtime/Interchange/Source/Dispatcher/Public/InterchangeDispatcher.h
@@ -0,0 +1,48 @@
+// Copyright Epic Games, Inc. All Rights Reserved.
+#pragma once
+
+#include "CoreMinimal.h"
+#include "HAL/CriticalSection.h"
+#include "HAL/Thread.h"
+#include "InterchangeDispatcherTask.h"
+#include "InterchangeWorkerHandler.h"
+
+namespace InterchangeDispatcher
+{
+
+// Handle a list of tasks, and a set of external workers to consume them.
+class INTERCHANGEDISPATCHER_API FInterchangeDispatcher
+{
+public:
+	FInterchangeDispatcher(const FString& InResultFolder);
+	~FInterchangeDispatcher() {TerminateProcess();}
+
+	int32 AddTask(const FString& JsonDescription);
+	TOptional<FTask> GetNextTask();
+	void SetTaskState(int32 TaskIndex, ETaskState TaskState, const FString& JsonResult, const FString& JSonMessages);
+	void GetTaskState(int32 TaskIndex, ETaskState& TaskState, FString& JsonResult, FString& JSonMessages);
+
+	void StartProcess();
+	void TerminateProcess();
+	void WaitAllTaskToCompleteExecution();
+	bool IsOver();
+
+private:
+	void SpawnHandler();
+	bool IsHandlerAlive();
+	void CloseHandler();
+
+	// Tasks
+	FCriticalSection TaskPoolCriticalSection;
+	TArray<FTask> TaskPool;
+	int32 NextTaskIndex;
+	int32 CompletedTaskCount;
+
+	/** Path where the result files are dump */
+	FString ResultFolder;
+
+	// Workers
+	TUniquePtr<FInterchangeWorkerHandler> WorkerHandler = nullptr;
+};
+
+} // NS InterchangeDispatcher
diff --git a/Engine/Plugins/Runtime/Interchange/Source/Dispatcher/Public/InterchangeDispatcherConfig.h b/Engine/Plugins/Runtime/Interchange/Source/Dispatcher/Public/InterchangeDispatcherConfig.h
new file mode 100644
--- /dev/null
+++ b/Engine/Plugins/Runtime/Interchange/Source/Dispatcher/Public/InterchangeDispatcherConfig.h
@@ -0,0 +1,34 @@
+// Copyright Epic Games, Inc. All Rights Reserved.
+
+#pragma once
+
+namespace InterchangeDispatcher
+{
+	namespace Config
+	{
+		// Server side timeout for connection with the client
+		constexpr double AcceptTimeout_s = 3.0;
+
+		// Client side timeout for connection with the server
+		constexpr double ConnectTimeout_s = 3.0;
+
+		// Maximum time we accept to wait when a command is sent.
+		// If the command is not sent by this delay, the operation is considered as failed
+		constexpr double SendCommandTimeout_s = 1.0;
+
+		// When true, a Terminate command is sent to the processor and we wait for
+		// TerminateTimeout_s seconds for it to close, before a hard TerminateProc.
+		// When false, the process is directly closed with a TerminateProc.
+		constexpr bool CloseProcessByCommand = true;
+
+		// Maximum time we accept to wait for the process to close after a Terminate
+		// command has been sent.
+		constexpr double TerminateTimeout_s = 1.0;
+
+		// Delay between successive checks for a new task, and validation of the connection
+		constexpr double IdleLoopDelay = 0.1;
+
+		// Delay between successive validation of the connection
+		constexpr double ProcessingLoopDelay = 0.5;
+	};
+} // ns InterchangeDispatcher
diff --git a/Engine/Plugins/Runtime/Interchange/Source/Dispatcher/Public/InterchangeDispatcherLog.h b/Engine/Plugins/Runtime/Interchange/Source/Dispatcher/Public/InterchangeDispatcherLog.h
new file mode 100644
--- /dev/null
+++ b/Engine/Plugins/Runtime/Interchange/Source/Dispatcher/Public/InterchangeDispatcherLog.h
@@ -0,0 +1,7 @@
+// Copyright Epic Games, Inc. All Rights Reserved.
+#pragma once
+
+#include "Logging/LogMacros.h"
+
+DECLARE_LOG_CATEGORY_EXTERN(LogInterchangeDispatcher, Log, All)
+
diff --git a/Engine/Plugins/Runtime/Interchange/Source/Dispatcher/Public/InterchangeDispatcherModule.h b/Engine/Plugins/Runtime/Interchange/Source/Dispatcher/Public/InterchangeDispatcherModule.h
new file mode 100644
--- /dev/null
+++ b/Engine/Plugins/Runtime/Interchange/Source/Dispatcher/Public/InterchangeDispatcherModule.h
@@ -0,0 +1,19 @@
+// Copyright Epic Games, Inc. All Rights Reserved.
+#pragma once
+
+#include "Modules/ModuleInterface.h"
+
+#define INTERCHANGEDISPATCHER_MODULE_NAME TEXT("InterchangeDispatcher")
+
+/**
+ * This module allow out of process interchange translator in case a third party SDK is not thread safe.
+ */
+class FInterchangeDispatcherModule : public IModuleInterface
+{
+public:
+	static FInterchangeDispatcherModule& Get();
+	static bool IsAvailable();
+
+private:
+	virtual void StartupModule() override;
+};
diff --git a/Engine/Plugins/Runtime/Interchange/Source/Dispatcher/Public/InterchangeDispatcherNetworking.h b/Engine/Plugins/Runtime/Interchange/Source/Dispatcher/Public/InterchangeDispatcherNetworking.h
new file mode 100644
--- /dev/null
+++ b/Engine/Plugins/Runtime/Interchange/Source/Dispatcher/Public/InterchangeDispatcherNetworking.h
@@ -0,0 +1,124 @@
+// Copyright Epic Games, Inc. All Rights Reserved.
+
+#pragma once
+
+#include "InterchangeCommands.h"
+
+#include "CoreMinimal.h"
+#include "Sockets.h"
+
+
+namespace InterchangeDispatcher
+{
+
+enum class SocketErrorCode
+{
+	NoError = 0,
+	Error_Create,
+	Error_Bind,
+	Error_Listen,
+	UnableToReadOnSocket,
+	UnableToSendData,
+	CouldNotStartWSA,
+	UnableToGetLocalAddress,
+	ConnectionToServerFailed,
+};
+
+
+struct FMessageHeader
+{
+	int32 ByteSize = -1;
+
+	friend void operator<< (FArchive& Ar, FMessageHeader& H)
+	{
+		FString Guard = Ar.IsLoading() ? TEXT("") : TEXT("guard");
+		Ar << Guard;
+		Ar << H.ByteSize;
+		ensure(Guard == TEXT("guard"));
+	}
+};
+
+
+class FNetworkNode
+{
+public:
+	virtual ~FNetworkNode();
+
+	// Send a byte buffer as an atomic message
+	// returns true when a message is successfully sent
+	bool SendMessage(const TArray<uint8>& Buffer, double Timeout_s);
+
+	// Receive a byte buffer as an atomic message
+	// returns true when a message is fetched. In that case, OutBuffer contains the message.
+	bool ReceiveMessage(TArray<uint8>& OutBuffer, double Timeout_s);
+
+	bool IsValid() const { return !bReadError && !bWriteError; }
+
+protected:
+	struct FMessage
+	{
+		FMessageHeader Header;
+		TArray<uint8> Content;
+	};
+	FSocket* CreateInternalSocket(const FString& Description);
+	void CloseSocket(FSocket*& Socket);
+	bool IsConnected();
+
+protected:
+	FSocket* ConnectedSocket = nullptr;
+	SocketErrorCode ConnectedSocketError = SocketErrorCode::NoError;
+
+	FMessage IncommingMessage;
+
+	bool bReadError = false;
+	bool bWriteError = false;
+};
+
+
+
+// Create, Bind, Listen, Accept sequence
+class FNetworkServerNode : public FNetworkNode
+{
+public:
+	FNetworkServerNode();
+	~FNetworkServerNode();
+	int32 GetListeningPort();
+	bool Accept(const FString& Description, double Timeout_s);
+
+private:
+	FSocket* ListeningSocket;
+};
+
+
+
+// Create, Connect sequence
+class FNetworkClientNode : public FNetworkNode
+{
+public:
+	bool Connect(const FString& Description, int32 ServerPort, double Timeout_s);
+};
+
+
+
+class FCommandQueue
+{
+public:
+	void SetNetworkInterface(FNetworkNode* InNetworkInterface);
+
+	TSharedPtr<ICommand> GetNextCommand(double Timeout_s);
+	bool SendCommand(ICommand& Commmand, double Timeout_s);
+
+	bool IsValid() { return NetworkInterface ? NetworkInterface->IsValid() : true;}
+
+	void Disconnect(double Timeout_s);
+
+private:
+	bool Poll(double Timeout_s);
+	TQueue<TSharedPtr<ICommand>> InCommands;
+
+private:
+	FNetworkNode* NetworkInterface = nullptr;
+};
+
+
+} // ns InterchangeDispatcher
diff --git a/Engine/Plugins/Runtime/Interchange/Source/Dispatcher/Public/InterchangeDispatcherTask.h b/Engine/Plugins/Runtime/Interchange/Source/Dispatcher/Public/InterchangeDispatcherTask.h
new file mode 100644
--- /dev/null
+++ b/Engine/Plugins/Runtime/Interchange/Source/Dispatcher/Public/InterchangeDispatcherTask.h
@@ -0,0 +1,148 @@
+// Copyright Epic Games, Inc. All Rights Reserved.
+
+#pragma once
+
+#include "CoreMinimal.h"
+
+namespace InterchangeDispatcher
+{
+enum class ETaskState
+{
+	Unknown,
+	Running,
+	UnTreated,
+	ProcessOk,
+	ProcessFailed,
+};
+
+struct FTask
+{
+	FTask() = default;
+
+	FTask(const FString& InJsonDescription)
+	{
+		JsonDescription = InJsonDescription;
+		State = ETaskState::UnTreated;
+	}
+
+	FString JsonDescription;
+	int32 Index = -1;
+	ETaskState State = ETaskState::Unknown;
+	FString JsonResult;
+	FString JsonMessages;
+};
+
+/**
+ * Json cmd helper to be able to read and write a FTask::JsonDescription
+ */
+class INTERCHANGEDISPATCHER_API IJsonCmdBase
+{
+public:
+	virtual ~IJsonCmdBase() = default;
+
+	virtual FString GetAction() const = 0;
+	virtual FString GetTranslatorID() const = 0;
+	virtual FString ToJson() const = 0;
+
+	/**
+	 * Return false if the JsonString do not match the command, true otherwise.
+	 */
+	virtual bool FromJson(const FString& JsonString) = 0;
+
+	static FString GetCommandIDKey()
+	{
+		static const FString Key = TEXT("CmdID");
+		return Key;
+	}
+	static FString GetTranslatorIDKey()
+	{
+		static const FString Key = TEXT("TranslatorID");
+		return Key;
+	}
+	static FString GetCommandDataKey()
+	{
+		static const FString Key = TEXT("CmdData");
+		return Key;
+	}
+
+protected:
+	//Use this member to know if the data is initialize before using it
+	bool bIsDataInitialize = false;
+};
+
+class INTERCHANGEDISPATCHER_API FJsonLoadSourceCmd : public IJsonCmdBase
+{
+public:
+	FJsonLoadSourceCmd()
+	{
+		bIsDataInitialize = false;
+	}
+
+	FJsonLoadSourceCmd(const FString& InTranslatorID, const FString& InSourceFilename)
+		: TranslatorID(InTranslatorID)
+		, SourceFilename(InSourceFilename)
+	{
+		bIsDataInitialize = true;
+	}
+
+	virtual FString GetAction() const override
+	{
+		static const FString LoadString = TEXT("Load");
+		return LoadString;
+	}
+
+	virtual FString GetTranslatorID() const override
+	{
+		//Code should not do query data if the data was not set before
+		ensure(bIsDataInitialize);
+		return TranslatorID;
+	}
+
+	virtual FString ToJson() const;
+	virtual bool FromJson(const FString& JsonString);
+
+	FString GetSourceFilename() const
+	{
+		//Code should not do query data if the data was not set before
+		ensure(bIsDataInitialize);
+		return SourceFilename;
+	}
+
+	static FString GetSourceFilenameKey()
+	{
+		static const FString Key = TEXT("SourceFile");
+		return Key;
+	}
+
+	/**
+	 * Use this class helper to create the cmd result json string and to read it
+	 */
+	class INTERCHANGEDISPATCHER_API JsonResultParser
+	{
+	public:
+		FString GetResultFilename() const
+		{
+			return ResultFilename;
+		}
+		void SetResultFilename(const FString& InResultFilename)
+		{
+			ResultFilename = InResultFilename;
+		}
+		FString ToJson() const;
+		bool FromJson(const FString& JsonString);
+
+		static FString GetResultFilenameKey()
+		{
+			const FString Key = TEXT("ResultFile");
+			return Key;
+		}
+	private:
+		FString ResultFilename = FString();
+	};
+
+private:
+	FString TranslatorID = FString();
+	FString SourceFilename = FString();
+};
+
+} // NS InterchangeDispatcher
diff --git a/Engine/Plugins/Runtime/Interchange/Source/Dispatcher/Public/InterchangeWorkerHandler.h b/Engine/Plugins/Runtime/Interchange/Source/Dispatcher/Public/InterchangeWorkerHandler.h
new file mode 100644
--- /dev/null
+++ b/Engine/Plugins/Runtime/Interchange/Source/Dispatcher/Public/InterchangeWorkerHandler.h
@@ -0,0 +1,81 @@
+// Copyright Epic Games, Inc. All Rights Reserved.
+
+#pragma once
+
+#include "CoreMinimal.h"
+
+#include "InterchangeCommands.h"
+#include "InterchangeDispatcherNetworking.h"
+#include "InterchangeDispatcherTask.h"
+
+#include "HAL/PlatformProcess.h"
+#include "HAL/Thread.h"
+
+
+namespace InterchangeDispatcher
+{
+
+class FInterchangeDispatcher;
+
+//Handle a Worker by socket communication
+class INTERCHANGEDISPATCHER_API FInterchangeWorkerHandler
+{
+	enum class EWorkerState
+	{
+		Uninitialized,
+		Idle, // Initialized, available for processing
+		Processing, // Currently processing a task
+		Closing, // in the process of terminating
+		Terminated, // aka. Not Alive
+	};
+
+	enum class EWorkerErrorState
+	{
+		Ok,
+		ConnectionFailed_NotBound,
+		ConnectionFailed_NoClient,
+		ConnectionLost,
+		ConnectionLost_SendFailed,
+		WorkerProcess_CantCreate,
+		WorkerProcess_Lost,
+	};
+
+public:
+	FInterchangeWorkerHandler(FInterchangeDispatcher& InDispatcher, FString& InResultFolder);
+	~FInterchangeWorkerHandler();
+
+	void Run();
+	bool IsAlive() const;
+	void Stop();
+
+private:
+	void RunInternal();
+	void StartWorkerProcess();
+	void ValidateConnection();
+
+	void ProcessCommand(ICommand& Command);
+	void ProcessCommand(FPingCommand& PingCommand);
+	void ProcessCommand(FCompletedTaskCommand& RunTaskCommand);
+	const TCHAR* EWorkerErrorStateAsString(EWorkerErrorState e);
+
+private:
+	FInterchangeDispatcher& Dispatcher;
+
+	// Send and receive commands
+	FNetworkServerNode NetworkInterface;
+	FCommandQueue CommandIO;
+	FThread IOThread;
+	FString ThreadName;
+
+	// External process
+	FProcHandle WorkerHandle;
+	TAtomic<EWorkerState> WorkerState;
+	EWorkerErrorState ErrorState;
+
+	// self
+	FString ResultFolder;
+	TOptional<FTask> CurrentTask;
+	bool bShouldTerminate;
+
+};
+} // namespace InterchangeDispatcher
diff --git a/Engine/Plugins/Runtime/Interchange/Source/Import/InterchangeImportPlugin.Build.cs b/Engine/Plugins/Runtime/Interchange/Source/Import/InterchangeImportPlugin.Build.cs
--- a/Engine/Plugins/Runtime/Interchange/Source/Import/InterchangeImportPlugin.Build.cs
+++ b/Engine/Plugins/Runtime/Interchange/Source/Import/InterchangeImportPlugin.Build.cs
@@ -13,7 +13,9 @@ namespace UnrealBuildTool.Rules
 					"CoreUObject",
 					"Engine",
 					"InterchangeCore",
+					"InterchangeDispatcher",
 					"InterchangeNodePlugin",
+					"MeshDescription"
 				}
 				);
 			PrivateDependencyModuleNames.AddRange(
@@ -21,6 +23,8 @@ namespace UnrealBuildTool.Rules
 				{
 					"ImageWrapper",
 					"RHI",
+					"InterchangeDispatcher",
+					"Json",
 				}
 				);
 		}
diff --git a/Engine/Plugins/Runtime/Interchange/Source/Import/Private/Fbx/InterchangeFbxTranslator.cpp b/Engine/Plugins/Runtime/Interchange/Source/Import/Private/Fbx/InterchangeFbxTranslator.cpp
new file mode 100644
--- /dev/null
+++ b/Engine/Plugins/Runtime/Interchange/Source/Import/Private/Fbx/InterchangeFbxTranslator.cpp
@@ -0,0 +1,155 @@
+// Copyright Epic Games, Inc. All Rights Reserved. 
+#include "Fbx/InterchangeFbxTranslator.h"
+
+#include "Dom/JsonValue.h"
+#include "Dom/JsonObject.h"
+#include "InterchangeDispatcher.h"
+#include "InterchangeDispatcherTask.h"
+#include "LogInterchangeImportPlugin.h"
+#include "Material/MaterialPayLoad.h"
+#include "Mesh/SkeletalMeshPayload.h"
+#include "Mesh/StaticMeshPayload.h"
+#include "Misc/ConfigCacheIni.h"
+#include "Misc/FileHelper.h"
+#include "Misc/Paths.h"
+#include "Modules/ModuleManager.h"
+#include "Nodes/BaseNodeContainer.h"
+#include "Serialization/JsonWriter.h"
+#include "Serialization/JsonSerializer.h"
+
+UInterchangeFbxTranslator::UInterchangeFbxTranslator(const class FObjectInitializer& ObjectInitializer)
+{
+	Dispatcher = nullptr;
+}
+
+bool UInterchangeFbxTranslator::CanImportSourceData(const UInterchangeSourceData* InSourceData) const
+{
+	const bool bIncludeDot = false;
+	FString Extension = FPaths::GetExtension(InSourceData->GetFilename(), bIncludeDot);
+	FString FbxExtension = (TEXT("fbx"));
+	return FbxExtension.StartsWith(Extension,ESearchCase::IgnoreCase);
+}
+
+bool UInterchangeFbxTranslator::Translate(const UInterchangeSourceData* SourceData, Interchange::FBaseNodeContainer& BaseNodeContainer) const
+{
+	FString Filename = SourceData->GetFilename();
+	if (!FPaths::FileExists(Filename))
+	{
+		return false;
+	}
+
+	if(!Dispatcher.IsValid())
+	{
+		//Dispatch an Interchange worker by using the InterchangeDispatcher
+		//Build Result folder
+		FGuid RandomGuid;
+		FPlatformMisc::CreateGuid(RandomGuid);
+		IPlatformFile& PlatformFile = FPlatformFileManager::Get().GetPlatformFile();
+		const FString ProjectSavedDir = FPaths::ConvertRelativePathToFull(FPaths::ProjectSavedDir());
+		const FString RamdomGuidDir = RandomGuid.ToString(EGuidFormats::Base36Encoded);
+		if (!PlatformFile.DirectoryExists(*ProjectSavedDir))
+		{
+			PlatformFile.CreateDirectory(*ProjectSavedDir);
+		}
+		const FString InterchangeDir = FPaths::Combine(ProjectSavedDir, TEXT("Interchange"));
+		if (!PlatformFile.DirectoryExists(*InterchangeDir))
+		{
+			PlatformFile.CreateDirectory(*InterchangeDir);
+		}
+		const FString ResultFolder = FPaths::Combine(InterchangeDir, RamdomGuidDir);
+		if (!PlatformFile.DirectoryExists(*ResultFolder))
+		{
+			PlatformFile.CreateDirectory(*ResultFolder);
+		}
+
+		//Create the dispatcher
+		Dispatcher = MakeUnique<InterchangeDispatcher::FInterchangeDispatcher>(ResultFolder);
+
+		if(ensure(Dispatcher.IsValid()))
+		{
+			Dispatcher->StartProcess();
+		}
+	}
+
+	if(!Dispatcher.IsValid())
+	{
+		return false;
+	}
+
+	//Create a json command to read the fbx file
+	FString JsonCommand = CreateLoadFbxFileCommand(SourceData->GetFilename());
+	int32 TaskIndex = Dispatcher->AddTask(JsonCommand);
+
+	//Blocking call until all tasks are executed
+	Dispatcher->WaitAllTaskToCompleteExecution();
+		
+	InterchangeDispatcher::ETaskState TaskState;
+	FString JsonResult;
+	FString JSonMessages;
+	Dispatcher->GetTaskState(TaskIndex, TaskState, JsonResult, JSonMessages);
+
+	//TODO: Parse the JSonMessage and add the message to the interchange not yet develop error messaging
+
+	if(TaskState != InterchangeDispatcher::ETaskState::ProcessOk)
+	{
+		return false;
+	}
+	//Grab the result file and fill the BaseNodeContainer
+	InterchangeDispatcher::FJsonLoadSourceCmd::JsonResultParser ResultParser;
+	ResultParser.FromJson(JsonResult);
+	FString BaseNodeContainerFilename = ResultParser.GetResultFilename();
+
+	//Parse the filename to fill the container
+
+	return true;
+}
+
+void UInterchangeFbxTranslator::ImportFinish()
+{
+	if (Dispatcher.IsValid())
+	{
+		Dispatcher->TerminateProcess();
+	}
+}
+
+TOptional<Interchange::FStaticMeshPayloadData> UInterchangeFbxTranslator::GetStaticMeshPayloadData(const UInterchangeSourceData* SourceData, const FString& PayLoadKey) const
+{
+	if (!SourceData)
+	{
+		UE_LOG(LogInterchangeImportPlugin, Error, TEXT("Failed to import Fbx, bad source data."));
+		return TOptional<Interchange::FStaticMeshPayloadData>();
+	}
+	//Use the Source data has the key for the dispatcher and send a command to retrieve the staticmesh bulk data
+	TOptional<Interchange::FStaticMeshPayloadData> PayloadData = Interchange::FStaticMeshPayloadData();
+	return PayloadData;
+}
+
+TOptional<Interchange::FSkeletalMeshPayloadData> UInterchangeFbxTranslator::GetSkeletalMeshPayloadData(const UInterchangeSourceData* SourceData, const FString& PayLoadKey) const
+{
+	if (!SourceData)
+	{
+		UE_LOG(LogInterchangeImportPlugin, Error, TEXT("Failed to import Fbx, bad source data."));
+		return TOptional<Interchange::FSkeletalMeshPayloadData>();
+	}
+	//Use the Source data has the key for the dispatcher and send a command to retrieve the skeletalmesh bulk data
+	TOptional<Interchange::FSkeletalMeshPayloadData> PayloadData = Interchange::FSkeletalMeshPayloadData();
+	return PayloadData;
+}
+
+TOptional<Interchange::FMaterialPayloadData> UInterchangeFbxTranslator::GetMaterialPayloadData(const UInterchangeSourceData* SourceData, const FString& PayLoadKey) const
+{
+	if (!SourceData)
+	{
+		UE_LOG(LogInterchangeImportPlugin, Error, TEXT("Failed to import Fbx, bad source data."));
+		return TOptional<Interchange::FMaterialPayloadData>();
+	}
+	//Use the Source data has the key for the dispatcher and send a command to retrieve the material bulk data
+	TOptional<Interchange::FMaterialPayloadData> PayloadData = Interchange::FMaterialPayloadData();
+	return PayloadData;
+}
+
+FString UInterchangeFbxTranslator::CreateLoadFbxFileCommand(const FString& FbxFilePath) const
+{
+	InterchangeDispatcher::FJsonLoadSourceCmd LoadSourceCommand(TEXT("FBX"), FbxFilePath);
+	return LoadSourceCommand.ToJson();
+}
\ No newline at end of file
diff --git a/Engine/Plugins/Runtime/Interchange/Source/Import/Private/InterchangeImportPlugin.cpp b/Engine/Plugins/Runtime/Interchange/Source/Import/Private/InterchangeImportPlugin.cpp
--- a/Engine/Plugins/Runtime/Interchange/Source/Import/Private/InterchangeImportPlugin.cpp
+++ b/Engine/Plugins/Runtime/Interchange/Source/Import/Private/InterchangeImportPlugin.cpp
@@ -1,6 +1,7 @@
 // Copyright Epic Games, Inc. All Rights Reserved.
 
 #include "CoreMinimal.h"
+#include "Fbx/InterchangeFbxTranslator.h"
 #include "IInterchangeImportPlugin.h"
 #include "InterchangeManager.h"
 #include "Modules/ModuleManager.h"
@@ -32,6 +33,8 @@ void FInterchangeImportPlugin::StartupModule()
 	UInterchangeManager& InterchangeManager = UInterchangeManager::GetInterchangeManager();
 
 	//Register the translators
+	//Scenes
+	InterchangeManager.RegisterTranslator(UInterchangeFbxTranslator::StaticClass()); //Do not submit uncommented until we replace completly fbx importer (staticmesh + skeletalMesh + animation)
 	//Textures
 	InterchangeManager.RegisterTranslator(UInterchangeBMPTranslator::StaticClass());
 	InterchangeManager.RegisterTranslator(UInterchangeDDSTranslator::StaticClass());
diff --git a/Engine/Plugins/Runtime/Interchange/Source/Import/Public/Fbx/InterchangeFbxTranslator.h b/Engine/Plugins/Runtime/Interchange/Source/Import/Public/Fbx/InterchangeFbxTranslator.h
new file mode 100644
--- /dev/null
+++ b/Engine/Plugins/Runtime/Interchange/Source/Import/Public/Fbx/InterchangeFbxTranslator.h
@@ -0,0 +1,70 @@
+// Copyright Epic Games, Inc. All Rights Reserved.
+
+#pragma once
+
+#include "CoreMinimal.h"
+#include "UObject/Object.h"
+#include "UObject/ObjectMacros.h"
+#include "InterchangeTranslatorBase.h"
+#include "InterchangeDispatcher.h"
+#include "Material/MaterialPayLoad.h"
+#include "Mesh/SkeletalMeshPayload.h"
+#include "Mesh/StaticMeshPayload.h"
+
+#include "InterchangeFbxTranslator.generated.h"
+
+UCLASS(BlueprintType)
+class INTERCHANGEIMPORTPLUGIN_API UInterchangeFbxTranslator : public UInterchangeTranslatorBase
+{
+	GENERATED_BODY()
+public:
+	UInterchangeFbxTranslator(const class FObjectInitializer& ObjectInitializer);
+
+	/** Begin UInterchangeTranslatorBase API*/
+	virtual bool CanImportSourceData(const UInterchangeSourceData* InSourceData) const override;
+	virtual bool Translate(const UInterchangeSourceData* SourceData, Interchange::FBaseNodeContainer& BaseNodeContainer) const override;
+	virtual void ImportFinish() override;
+	/** End UInterchangeTranslatorBase API*/
+
+	/**
+	 * Once the translation is done, the import process need a way to retrieve payload data.
+	 * This payload will be use by the factories to create the asset.
+	 *
+	 * @param SourceData - The source data containing the data to translate
+	 * @param PayloadKey - The key to retrieve the a particular payload contain into the specified source data.
+	 * @return a PayloadData containing the mesh data to import. The TOptional will not be set if there is an error.
+	 */
+	virtual TOptional<Interchange::FStaticMeshPayloadData> GetStaticMeshPayloadData(const UInterchangeSourceData* SourceData, const FString& PayLoadKey) const;
+
+	/**
+	 * Once the translation is done, the import process need a way to retrieve payload data.
+	 * This payload will be use by the factories to create the asset.
+	 *
+	 * @param SourceData - The source data containing the data to translate
+	 * @param PayloadKey - The key to retrieve the a particular payload contain into the specified source data.
+	 * @return a PayloadData containing the mesh data to import. The TOptional will not be set if there is an error.
+	 */
+	virtual TOptional<Interchange::FSkeletalMeshPayloadData> GetSkeletalMeshPayloadData(const UInterchangeSourceData* SourceData, const FString& PayLoadKey) const;
+
+	/**
+	 * Once the translation is done, the import process need a way to retrieve payload data.
+	 * This payload will be use by the factories to create the asset.
+	 *
+	 * @param SourceData - The source data containing the data to translate
+	 * @param PayloadKey - The key to retrieve the a particular payload contain into the specified source data.
+	 * @return a PayloadData containing the mesh data to import. The TOptional will not be set if there is an error.
+	 */
+	virtual TOptional<Interchange::FMaterialPayloadData> GetMaterialPayloadData(const UInterchangeSourceData* SourceData, const FString& PayLoadKey) const;
+
+private:
+
+	FString CreateLoadFbxFileCommand(const FString& FbxFilePath) const;
+
+	//Dispatcher is mutable since it is create during the Translate operation
+	//We do not want to allocate the dispatcher and start the InterchangeWorker process
+	//in the constructor because Archetype, CDO and registered translators will
+	//never translate a source.
+	mutable TUniquePtr<InterchangeDispatcher::FInterchangeDispatcher> Dispatcher;
+};
+
+
diff --git a/Engine/Plugins/Runtime/Interchange/Source/Import/Public/Material/MaterialPayLoad.h b/Engine/Plugins/Runtime/Interchange/Source/Import/Public/Material/MaterialPayLoad.h
new file mode 100644
--- /dev/null
+++ b/Engine/Plugins/Runtime/Interchange/Source/Import/Public/Material/MaterialPayLoad.h
@@ -0,0 +1,13 @@
+// Copyright Epic Games, Inc. All Rights Reserved.
+
+#pragma once
+
+#include "CoreMinimal.h"
+
+namespace Interchange
+{
+	struct FMaterialPayloadData
+	{
+		FColor BaseColor;
+	};
+}
diff --git a/Engine/Plugins/Runtime/Interchange/Source/Import/Public/Mesh/SkeletalMeshPayload.h b/Engine/Plugins/Runtime/Interchange/Source/Import/Public/Mesh/SkeletalMeshPayload.h
new file mode 100644
--- /dev/null
+++ b/Engine/Plugins/Runtime/Interchange/Source/Import/Public/Mesh/SkeletalMeshPayload.h
@@ -0,0 +1,16 @@
+// Copyright Epic Games, Inc. All Rights Reserved.
+
+#pragma once
+
+#include "CoreMinimal.h"
+#include "Rendering/SkeletalMeshLODImporterData.h"
+
+namespace Interchange
+{
+	struct FSkeletalMeshPayloadData
+	{
+#if WITH_EDITOR
+		TArray<FSkeletalMeshImportData> SkeletalMeshImportData;
+#endif
+	};
+}
diff --git a/Engine/Plugins/Runtime/Interchange/Source/Import/Public/Mesh/StaticMeshPayload.h b/Engine/Plugins/Runtime/Interchange/Source/Import/Public/Mesh/StaticMeshPayload.h
new file mode 100644
--- /dev/null
+++ b/Engine/Plugins/Runtime/Interchange/Source/Import/Public/Mesh/StaticMeshPayload.h
@@ -0,0 +1,14 @@
+// Copyright Epic Games, Inc. All Rights Reserved.
+
+#pragma once
+
+#include "CoreMinimal.h"
+#include "MeshDescription.h"
+
+namespace Interchange
+{
+	struct FStaticMeshPayloadData
+	{
+		TArray<FMeshDescription> MeshDescriptions;
+	};
+}
diff --git a/Engine/Plugins/Runtime/Interchange/Source/Nodes/InterchangeNodePlugin.Build.cs b/Engine/Plugins/Runtime/Interchange/Source/Nodes/InterchangeNodePlugin.Build.cs
--- a/Engine/Plugins/Runtime/Interchange/Source/Nodes/InterchangeNodePlugin.Build.cs
+++ b/Engine/Plugins/Runtime/Interchange/Source/Nodes/InterchangeNodePlugin.Build.cs
@@ -14,6 +14,12 @@ namespace UnrealBuildTool.Rules
 					"InterchangeCore",
 				}
 			);
+
+			if(Target.bCompileAgainstEngine)
+            {
+				PublicDependencyModuleNames.Add("Engine");
+
+			}
 		}
 	}
 }
diff --git a/Engine/Plugins/Runtime/Interchange/Source/Nodes/Public/TextureNode.h b/Engine/Plugins/Runtime/Interchange/Source/Nodes/Public/TextureNode.h
--- a/Engine/Plugins/Runtime/Interchange/Source/Nodes/Public/TextureNode.h
+++ b/Engine/Plugins/Runtime/Interchange/Source/Nodes/Public/TextureNode.h
@@ -4,8 +4,11 @@
 
 #include "CoreMinimal.h"
 #include "Nodes/BaseNode.h"
+
+#if WITH_ENGINE
 #include "Engine/Texture.h"
 #include "Engine/Texture2D.h"
+#endif
 
 //Interchange namespace
 namespace Interchange
diff --git a/Engine/Plugins/Runtime/Interchange/Source/Parsers/Fbx/InterchangeFbxParser.Build.cs b/Engine/Plugins/Runtime/Interchange/Source/Parsers/Fbx/InterchangeFbxParser.Build.cs
new file mode 100644
--- /dev/null
+++ b/Engine/Plugins/Runtime/Interchange/Source/Parsers/Fbx/InterchangeFbxParser.Build.cs
@@ -0,0 +1,32 @@
+// Copyright Epic Games, Inc. All Rights Reserved.
+
+using System.IO;
+
+namespace UnrealBuildTool.Rules
+{
+	public class InterchangeFbxParser : ModuleRules
+	{
+		public InterchangeFbxParser(ReadOnlyTargetRules Target) : base(Target)
+		{
+            PublicDependencyModuleNames.AddRange(
+				new string[]
+				{
+                    "Core",
+					"CoreUObject",
+					"InterchangeCore",
+				}
+			);
+
+			PrivateDependencyModuleNames.AddRange(
+				new string[]
+				{
+					"InterchangeNodePlugin",
+				}
+			);
+
+			AddEngineThirdPartyPrivateStaticDependencies(Target,
+				"FBX"
+			);
+		}
+    }
+}
diff --git a/Engine/Plugins/Runtime/Interchange/Source/Parsers/Fbx/Private/InterchangeFbxParser.cpp b/Engine/Plugins/Runtime/Interchange/Source/Parsers/Fbx/Private/InterchangeFbxParser.cpp
new file mode 100644
--- /dev/null
+++ b/Engine/Plugins/Runtime/Interchange/Source/Parsers/Fbx/Private/InterchangeFbxParser.cpp
@@ -0,0 +1,16 @@
+// Copyright Epic Games, Inc. All Rights Reserved.
+
+#include "InterchangeFbxParser.h"
+
+#include "CoreMinimal.h"
+
+namespace InterchangeFbxParser
+{
+	void FbxParser::LoadFbxFile(const FString& Filename, const FString& ResultFolder)
+	{
+		ResultFilepath = ResultFolder + TEXT("/Foo.txt");
+		JsonLoadMessages = TEXT("{\"Msg\" : {\"Type\" : \"Log\",\n\"Msg\" : \"This is a success!\"}}");
+		return;
+	}
+
+}
diff --git a/Engine/Plugins/Runtime/Interchange/Source/Parsers/Fbx/Private/InterchangeFbxParserModule.cpp b/Engine/Plugins/Runtime/Interchange/Source/Parsers/Fbx/Private/InterchangeFbxParserModule.cpp
new file mode 100644
--- /dev/null
+++ b/Engine/Plugins/Runtime/Interchange/Source/Parsers/Fbx/Private/InterchangeFbxParserModule.cpp
@@ -0,0 +1,32 @@
+// Copyright Epic Games, Inc. All Rights Reserved.
+#include "InterchangeFbxParserModule.h"
+
+#include "InterchangeFbxParserLog.h"
+
+#include "Modules/ModuleInterface.h"
+#include "Modules/ModuleManager.h"
+
+
+
+#define LOCTEXT_NAMESPACE "FInterchangeFbxParserModule"
+
+DEFINE_LOG_CATEGORY(LogInterchangeFbxParser);
+
+FInterchangeFbxParserModule& FInterchangeFbxParserModule::Get()
+{
+	return FModuleManager::LoadModuleChecked< FInterchangeFbxParserModule >(INTERCHANGEFBXPARSER_MODULE_NAME);
+}
+
+bool FInterchangeFbxParserModule::IsAvailable()
+{
+	return FModuleManager::Get().IsModuleLoaded(INTERCHANGEFBXPARSER_MODULE_NAME);
+}
+
+void FInterchangeFbxParserModule::StartupModule()
+{
+}
+
+IMPLEMENT_MODULE(FInterchangeFbxParserModule, InterchangeFbxParser);
+
+#undef LOCTEXT_NAMESPACE // "InterchangeFbxParserModule"
+
diff --git a/Engine/Plugins/Runtime/Interchange/Source/Parsers/Fbx/Public/InterchangeFbxParser.h b/Engine/Plugins/Runtime/Interchange/Source/Parsers/Fbx/Public/InterchangeFbxParser.h
new file mode 100644
--- /dev/null
+++ b/Engine/Plugins/Runtime/Interchange/Source/Parsers/Fbx/Public/InterchangeFbxParser.h
@@ -0,0 +1,30 @@
+// Copyright Epic Games, Inc. All Rights Reserved.
+
+#pragma once
+
+#include "CoreMinimal.h"
+#include "Nodes/BaseNodeContainer.h"
+
+namespace InterchangeFbxParser
+{
+
+class FbxParser
+{
+public:
+	/**
+	 * Parse a file support by the fbx sdk. It just extract all the fbx node and create a FBaseNodeContainer and dump it in a json file inside the ResultFolder
+	 * @param - Filename is the file that the fbx sdk will read (.fbx or .obj)
+	 * @param - ResultFolder is the folder where we must put any result file
+	 */
+	void LoadFbxFile(const FString& Filename, const FString& ResultFolder);
+	
+	FString GetResultFilepath(){ return ResultFilepath; }
+	FString GetJsonLoadMessages() { return JsonLoadMessages; }
+
+private:
+	Interchange::FBaseNodeContainer Container;
+	FString ResultFilepath;
+	FString JsonLoadMessages;
+};
+
+} // ns InterchangeFbxParser
diff --git a/Engine/Plugins/Runtime/Interchange/Source/Parsers/Fbx/Public/InterchangeFbxParserLog.h b/Engine/Plugins/Runtime/Interchange/Source/Parsers/Fbx/Public/InterchangeFbxParserLog.h
new file mode 100644
--- /dev/null
+++ b/Engine/Plugins/Runtime/Interchange/Source/Parsers/Fbx/Public/InterchangeFbxParserLog.h
@@ -0,0 +1,7 @@
+// Copyright Epic Games, Inc. All Rights Reserved.
+#pragma once
+
+#include "Logging/LogMacros.h"
+
+DECLARE_LOG_CATEGORY_EXTERN(LogInterchangeFbxParser, Log, All)
+
diff --git a/Engine/Plugins/Runtime/Interchange/Source/Parsers/Fbx/Public/InterchangeFbxParserModule.h b/Engine/Plugins/Runtime/Interchange/Source/Parsers/Fbx/Public/InterchangeFbxParserModule.h
new file mode 100644
--- /dev/null
+++ b/Engine/Plugins/Runtime/Interchange/Source/Parsers/Fbx/Public/InterchangeFbxParserModule.h
@@ -0,0 +1,19 @@
+// Copyright Epic Games, Inc. All Rights Reserved.
+#pragma once
+
+#include "Modules/ModuleInterface.h"
+
+#define INTERCHANGEFBXPARSER_MODULE_NAME TEXT("InterchangeFbxParser")
+
+/**
+ * This module allow out of process interchange translator in case a third party SDK is not thread safe.
+ */
+class FInterchangeFbxParserModule : public IModuleInterface
+{
+public:
+	static FInterchangeFbxParserModule& Get();
+	static bool IsAvailable();
+
+private:
+	virtual void StartupModule() override;
+};
diff --git a/Engine/Source/Programs/InterchangeWorker/InterchangeWorker.Build.cs b/Engine/Source/Programs/InterchangeWorker/InterchangeWorker.Build.cs
new file mode 100644
--- /dev/null
+++ b/Engine/Source/Programs/InterchangeWorker/InterchangeWorker.Build.cs
@@ -0,0 +1,29 @@
+// Copyright Epic Games, Inc. All Rights Reserved.
+
+using UnrealBuildTool;
+using System.IO;
+
+public class InterchangeWorker : ModuleRules
+{
+	public InterchangeWorker(ReadOnlyTargetRules Target) : base(Target)
+	{
+
+		PublicIncludePaths.Add("Runtime/Launch/Public");
+		PrivateIncludePaths.Add("Runtime/Launch/Private");
+
+		PrivateDependencyModuleNames.AddRange(
+			new string[]
+			{
+				"ApplicationCore",
+				"Core",
+				"InterchangeDispatcher",
+				"InterchangeFbxParser",
+				"Json",
+				"Projects",
+				"Sockets",
+			}
+		);
+
+		PublicDelayLoadDLLs.Add("kernel_io.dll");
+	}
+}
diff --git a/Engine/Source/Programs/InterchangeWorker/InterchangeWorker.Target.cs b/Engine/Source/Programs/InterchangeWorker/InterchangeWorker.Target.cs
new file mode 100644
--- /dev/null
+++ b/Engine/Source/Programs/InterchangeWorker/InterchangeWorker.Target.cs
@@ -0,0 +1,45 @@
+// Copyright Epic Games, Inc. All Rights Reserved.
+
+using UnrealBuildTool;
+using System.Collections.Generic;
+
+[SupportedPlatforms(UnrealPlatformClass.Desktop)]
+public class InterchangeWorkerTarget : TargetRules
+{
+	public InterchangeWorkerTarget(TargetInfo Target) : base(Target)
+	{
+		Type = TargetType.Program;
+		LinkType = TargetLinkType.Monolithic;
+		LaunchModuleName = "InterchangeWorker";
+		SolutionDirectory = "Programs/InterchangeWorker";
+
+		// Lean and mean
+		bBuildDeveloperTools = false;
+
+		// Never use malloc profiling in Unreal Header Tool.  We set this because often UHT is compiled right before the engine
+		// automatically by Unreal Build Tool, but if bUseMallocProfiler is defined, UHT can operate incorrectly.
+		bUseMallocProfiler = false;
+
+		// Editor-only data, however, is needed
+		bBuildWithEditorOnlyData = false;
+
+		// Currently this app is not linking against the engine, so we'll compile out references from Core to the rest of the engine
+		bCompileAgainstEngine = false;
+		bCompileAgainstCoreUObject = true;
+		bCompileAgainstApplicationCore = true;
+
+		// This is a console application, not a Windows app (sets entry point to main(), instead of WinMain())
+		bIsBuildingConsoleApplication = true;
+
+		bLegalToDistributeBinary = true;
+
+		bCompilePhysX = false;
+		bCompileAPEX = false;
+		bCompileNvCloth = false;
+		bCompileICU = false;
+		bCompileCEF3 = false;
+
+		string EngineDir = @"..\.."; // relative to default destination, which is $(EngineDir)\Binaries\Win64
+		ExeBinariesSubFolder = EngineDir + @"\Plugins\Runtime\Interchange\Binaries\Win64\";
+	}
+}
diff --git a/Engine/Source/Programs/InterchangeWorker/Private/InterchangeWorker.cpp b/Engine/Source/Programs/InterchangeWorker/Private/InterchangeWorker.cpp
new file mode 100644
--- /dev/null
+++ b/Engine/Source/Programs/InterchangeWorker/Private/InterchangeWorker.cpp
@@ -0,0 +1,92 @@
+// Copyright Epic Games, Inc. All Rights Reserved.
+
+#include "InterchangeWorker.h"
+#include "InterchangeWorkerImpl.h"
+
+#include "RequiredProgramMainCPPInclude.h"
+
+IMPLEMENT_APPLICATION(InterchangeWorker, "InterchangeWorker");
+DEFINE_LOG_CATEGORY(LogInterchangeWorker);
+
+#define EXIT_MISSING_CORETECH_MODULE 2
+
+bool GIsConsoleExecutable = true;
+
+void GetParameter(int32 Argc, TCHAR* Argv[], const FString& InParam, FString& OutValue)
+{
+	OutValue.Empty();
+	for (int32 Index = 1; Index < Argc; Index++)
+	{
+		if (! FCString::Strcmp(*InParam, Argv[Index]))
+		{
+			if (Index + 1 < Argc)
+			{
+				OutValue = Argv[Index + 1];
+			}
+		}
+	}
+}
+
+bool HasParameter(int32 Argc, TCHAR* Argv[], const FString& InParam)
+{
+	for (int32 Index = 1; Index < Argc; Index++)
+	{
+		if (!FCString::Strcmp(*InParam, Argv[Index]))
+		{
+			return true;
+		}
+	}
+	return false;
+}
+
+int32 Main(int32 Argc, TCHAR * Argv[])
+{
+	UE_SET_LOG_VERBOSITY(LogInterchangeWorker, Verbose);
+
+	FString ResultFolder, ServerPID, ServerPort, InterchangeDispatcherVersion;
+	GetParameter(Argc, Argv, "-ServerPID", ServerPID);
+	GetParameter(Argc, Argv, "-ServerPort", ServerPort);
+	GetParameter(Argc, Argv, "-InterchangeDispatcherVersion", InterchangeDispatcherVersion);
+	GetParameter(Argc, Argv, "-ResultFolder", ResultFolder);
+
+	int32 Major = 0;
+	int32 Minor = 0;
+	int32 Patch = 0;
+	if (!InterchangeDispatcher::DispatcherCommandVersion::FromString(InterchangeDispatcherVersion, Major, Minor, Patch))
+	{
+		UE_LOG(LogInterchangeWorker, Error, TEXT("Incompatible interchange dispatcher version string command argument."));
+		return EXIT_FAILURE;
+	}
+ 	if (!InterchangeDispatcher::DispatcherCommandVersion::IsAPICompatible(Major, Minor, Patch))
+ 	{
+ 		UE_LOG(LogInterchangeWorker, Error, TEXT("Incompatible interchange dispatcher version. Please recompile InterchangeWorker target."));
+ 		return EXIT_FAILURE;
+ 	}
+
+	FInterchangeWorkerImpl Worker(FCString::Atoi(*ServerPID), FCString::Atoi(*ServerPort), ResultFolder);
+	Worker.Run();
+
+	return EXIT_SUCCESS;
+}
+
+int32 Filter(uint32 Code, struct _EXCEPTION_POINTERS *Ep)
+{
+	return EXCEPTION_EXECUTE_HANDLER;
+}
+
+INT32_MAIN_INT32_ARGC_TCHAR_ARGV()
+{
+	GEngineLoop.PreInit(ArgC, ArgV);
+
+	SetErrorMode(SEM_FAILCRITICALERRORS | SEM_NOGPFAULTERRORBOX);
+	_set_abort_behavior(0, _WRITE_ABORT_MSG);
+	__try
+	{
+		return Main(ArgC, ArgV);
+	}
+	__except (Filter(GetExceptionCode(), GetExceptionInformation()))
+	{
+		return EXIT_FAILURE;
+	}
+	return EXIT_SUCCESS;
+}
diff --git a/Engine/Source/Programs/InterchangeWorker/Private/InterchangeWorker.h b/Engine/Source/Programs/InterchangeWorker/Private/InterchangeWorker.h
new file mode 100644
--- /dev/null
+++ b/Engine/Source/Programs/InterchangeWorker/Private/InterchangeWorker.h
@@ -0,0 +1,8 @@
+// Copyright Epic Games, Inc. All Rights Reserved.
+
+#pragma once
+
+#include "CoreMinimal.h"
+#include "Logging/LogMacros.h"
+
+DECLARE_LOG_CATEGORY_EXTERN(LogInterchangeWorker, Log, All);
diff --git a/Engine/Source/Programs/InterchangeWorker/Private/InterchangeWorkerImpl.cpp b/Engine/Source/Programs/InterchangeWorker/Private/InterchangeWorkerImpl.cpp
new file mode 100644
--- /dev/null
+++ b/Engine/Source/Programs/InterchangeWorker/Private/InterchangeWorkerImpl.cpp
@@ -0,0 +1,166 @@
+// Copyright Epic Games, Inc. All Rights Reserved.
+
+#include "InterchangeWorkerImpl.h"
+
+#include "InterchangeCommands.h"
+#include "InterchangeDispatcherConfig.h"
+#include "InterchangeDispatcherTask.h"
+#include "InterchangeFbxParser.h"
+
+#include "HAL/FileManager.h"
+#include "Misc/FileHelper.h"
+#include "Misc/Paths.h"
+#include "SocketSubsystem.h"
+#include "Sockets.h"
+
+using namespace InterchangeDispatcher;
+
+
+FInterchangeWorkerImpl::FInterchangeWorkerImpl(int32 InServerPID, int32 InServerPort, FString& InResultFolder)
+	: ServerPID(InServerPID)
+	, ServerPort(InServerPort)
+	, PingStartCycle(0)
+	, ResultFolder(InResultFolder)
+{
+	IPlatformFile& PlatformFile = FPlatformFileManager::Get().GetPlatformFile();
+	FPaths::NormalizeDirectoryName(ResultFolder);
+	if (!PlatformFile.DirectoryExists(*ResultFolder))
+	{
+		PlatformFile.CreateDirectory(*ResultFolder);
+	}
+}
+
+bool FInterchangeWorkerImpl::Run()
+{
+	UE_LOG(LogInterchangeWorker, Verbose, TEXT("connect to %d..."), ServerPort);
+	bool bConnected = NetworkInterface.Connect(TEXT("Interchange Worker"), ServerPort, Config::ConnectTimeout_s);
+	UE_LOG(LogInterchangeWorker, Verbose, TEXT("connected to %d %s"), ServerPort, bConnected ? TEXT("OK") : TEXT("FAIL"));
+	if (bConnected)
+	{
+		CommandIO.SetNetworkInterface(&NetworkInterface);
+	}
+	else
+	{
+		UE_LOG(LogInterchangeWorker, Error, TEXT("Server connection failure. exit"));
+		return false;
+	}
+
+	InitiatePing();
+
+	bool bIsRunning = true;
+	while (bIsRunning)
+	{
+		if (TSharedPtr<ICommand> Command = CommandIO.GetNextCommand(1.0))
+		{
+			switch(Command->GetType())
+			{
+				case ECommandId::Ping:
+					ProcessCommand(*StaticCast<FPingCommand*>(Command.Get()));
+					break;
+
+				case ECommandId::BackPing:
+					ProcessCommand(*StaticCast<FBackPingCommand*>(Command.Get()));
+					break;
+
+				case ECommandId::RunTask:
+					ProcessCommand(*StaticCast<FRunTaskCommand*>(Command.Get()));
+					break;
+
+				case ECommandId::Terminate:
+					UE_LOG(LogInterchangeWorker, Verbose, TEXT("Terminate command received. Exiting."));
+					bIsRunning = false;
+					break;
+
+				case ECommandId::NotifyEndTask:
+				default:
+					break;
+			}
+		}
+		else
+		{
+			if (bIsRunning)
+			{
+				bIsRunning = ServerPID == 0 ? true : FPlatformProcess::IsApplicationRunning(ServerPID);
+				UE_CLOG(!bIsRunning, LogInterchangeWorker, Error, TEXT("Worker failure: server lost"));
+			}
+		}
+	}
+
+	UE_CLOG(!bIsRunning, LogInterchangeWorker, Verbose, TEXT("Worker loop exit..."));
+	CommandIO.Disconnect(0);
+	return true;
+}
+
+void FInterchangeWorkerImpl::InitiatePing()
+{
+	PingStartCycle = FPlatformTime::Cycles64();
+	FPingCommand Ping;
+	CommandIO.SendCommand(Ping, Config::SendCommandTimeout_s);
+}
+
+void FInterchangeWorkerImpl::ProcessCommand(const FPingCommand& PingCommand)
+{
+	FBackPingCommand BackPing;
+	CommandIO.SendCommand(BackPing, Config::SendCommandTimeout_s);
+}
+
+void FInterchangeWorkerImpl::ProcessCommand(const FBackPingCommand& BackPingCommand)
+{
+	if (PingStartCycle)
+	{
+		double ElapsedTime_s = FGenericPlatformTime::ToSeconds(FPlatformTime::Cycles64() - PingStartCycle);
+		UE_LOG(LogInterchangeWorker, Verbose, TEXT("Ping %f s"), ElapsedTime_s);
+	}
+	PingStartCycle = 0;
+}
+
+void FInterchangeWorkerImpl::ProcessCommand(const FRunTaskCommand& RunTaskCommand)
+{
+	const FString& JsonToProcess = RunTaskCommand.JsonDescription;
+	UE_LOG(LogInterchangeWorker, Verbose, TEXT("Process %s"), *JsonToProcess);
+
+	ETaskState ProcessResult = ETaskState::Unknown;
+	//Process the json and run the task
+	FString JSonResult;
+	FString JSonMessages;
+	InterchangeDispatcher::FJsonLoadSourceCmd LoadSourceCommand;
+	//Any command FromJson function return true if the Json descibe the command
+	if (LoadSourceCommand.FromJson(JsonToProcess))
+	{
+		//Load file command
+		if (LoadSourceCommand.GetTranslatorID().Equals(TEXT("FBX"), ESearchCase::IgnoreCase))
+		{
+			//We want to load an FBX file
+			ProcessResult = LoadFbxFile(LoadSourceCommand, JSonResult, JSonMessages);
+		}
+	}
+	//TODO add GetPayload command
+	else
+	{
+		ProcessResult = ETaskState::Unknown;
+	}
+
+	FCompletedTaskCommand CompletedTask;
+	CompletedTask.ProcessResult = ProcessResult;
+	CompletedTask.JSonMessages = JSonMessages;
+	if (CompletedTask.ProcessResult == ETaskState::ProcessOk)
+	{
+		CompletedTask.JSonResult = JSonResult;
+	}
+
+	CommandIO.SendCommand(CompletedTask, Config::SendCommandTimeout_s);
+
+	UE_LOG(LogInterchangeWorker, Verbose, TEXT("End of Process %s"), *JsonToProcess);
+}
+
+ETaskState FInterchangeWorkerImpl::LoadFbxFile(const InterchangeDispatcher::FJsonLoadSourceCmd& LoadSourceCommand, FString& OutJSonResult, FString& OutJSonMessages) const
+{
+	FString SourceFilename = LoadSourceCommand.GetSourceFilename();
+	InterchangeFbxParser::FbxParser FbxParser;
+	FbxParser.LoadFbxFile(SourceFilename, ResultFolder);
+	InterchangeDispatcher::FJsonLoadSourceCmd::JsonResultParser ResultParser;
+	ResultParser.SetResultFilename(FbxParser.GetResultFilepath());
+	OutJSonMessages = FbxParser.GetJsonLoadMessages();
+	OutJSonResult = ResultParser.ToJson();
+	return ETaskState::ProcessOk;
+}
\ No newline at end of file
diff --git a/Engine/Source/Programs/InterchangeWorker/Private/InterchangeWorkerImpl.h b/Engine/Source/Programs/InterchangeWorker/Private/InterchangeWorkerImpl.h
new file mode 100644
--- /dev/null
+++ b/Engine/Source/Programs/InterchangeWorker/Private/InterchangeWorkerImpl.h
@@ -0,0 +1,37 @@
+// Copyright Epic Games, Inc. All Rights Reserved.
+#pragma once
+
+#include "InterchangeWorker.h"
+
+#include "InterchangeCommands.h"
+#include "InterchangeDispatcherNetworking.h"
+
+#include "HAL/PlatformTime.h"
+
+
+struct FFileStatData;
+struct FImportParameters;
+
+class FInterchangeWorkerImpl
+{
+public:
+	FInterchangeWorkerImpl(int32 InServerPID, int32 InServerPort, FString& InResultFolder);
+	bool Run();
+
+private:
+	void InitiatePing();
+	void ProcessCommand(const InterchangeDispatcher::FPingCommand& PingCommand);
+	void ProcessCommand(const InterchangeDispatcher::FBackPingCommand& BackPingCommand);
+	void ProcessCommand(const InterchangeDispatcher::FRunTaskCommand& TerminateCommand);
+
+	InterchangeDispatcher::ETaskState LoadFbxFile(const InterchangeDispatcher::FJsonLoadSourceCmd& LoadSourceCommand, FString& OutJSonResult, FString& OutJSonMessages) const;
+
+private:
+	InterchangeDispatcher::FNetworkClientNode NetworkInterface;
+	InterchangeDispatcher::FCommandQueue CommandIO;
+
+	int32 ServerPID;
+	int32 ServerPort;
+	uint64 PingStartCycle;
+	FString ResultFolder;
+};
diff --git a/Engine/Source/Runtime/Interchange/Core/InterchangeCore.Build.cs b/Engine/Source/Runtime/Interchange/Core/InterchangeCore.Build.cs
--- a/Engine/Source/Runtime/Interchange/Core/InterchangeCore.Build.cs
+++ b/Engine/Source/Runtime/Interchange/Core/InterchangeCore.Build.cs
@@ -10,14 +10,17 @@ public class InterchangeCore : ModuleRules
 				new string[] {
 					"Core",
 					"CoreUObject",
-					"Engine",
 				}
 			);
 
-		PrivateDependencyModuleNames.AddRange(
+		if(Target.bCompileAgainstEngine)
+        {
+			PrivateDependencyModuleNames.AddRange(
 				new string[] {
 					"AssetRegistry",
+					"Engine",
 				}
 			);
+		}
 	}
 }
diff --git a/Engine/Source/Runtime/Interchange/Core/Private/InterchangeManager.cpp b/Engine/Source/Runtime/Interchange/Core/Private/InterchangeManager.cpp
--- a/Engine/Source/Runtime/Interchange/Core/Private/InterchangeManager.cpp
+++ b/Engine/Source/Runtime/Interchange/Core/Private/InterchangeManager.cpp
@@ -2,10 +2,7 @@
 
 #include "InterchangeManager.h"
 
-#include "AssetDataTagMap.h"
-#include "AssetRegistryModule.h"
 #include "CoreMinimal.h"
-#include "Engine/Blueprint.h"
 #include "InterchangeFactoryBase.h"
 #include "InterchangeLogPrivate.h"
 #include "InterchangeSourceData.h"
@@ -28,6 +25,12 @@
 #include "UObject/UObjectIterator.h"
 #include "UObject/WeakObjectPtrTemplates.h"
 
+#if WITH_ENGINE
+#include "AssetDataTagMap.h"
+#include "AssetRegistryModule.h"
+#include "Engine/Blueprint.h"
+#endif //WITH_ENGINE
+
 namespace InternalInterchangePrivate
 {
 	const FLogCategoryBase* GetLogInterchangePtr()
@@ -103,6 +106,7 @@ void Interchange::FImportAsyncHelper::CleanUp()
 	{
 		if(Translator)
 		{
+			Translator->ImportFinish();
 			Translator->RemoveFromRoot();
 			Translator->MarkPendingKill();
 			Translator = nullptr;
@@ -438,6 +442,8 @@ void UInterchangeManager::FindPipelineCandidate(TArray<UClass*>& PipelineCandida
 		PipelineCandidates.AddUnique(Class);
 	}
 
+//Blueprint and python script discoverability is available only if we compile with the engine
+#if WITH_ENGINE
 	// Load the asset registry module
 	FAssetRegistryModule& AssetRegistryModule = FModuleManager::LoadModuleChecked< FAssetRegistryModule >(FName("AssetRegistry"));
 	IAssetRegistry& AssetRegistry = AssetRegistryModule.Get();
@@ -490,4 +496,6 @@ void UInterchangeManager::FindPipelineCandidate(TArray<UClass*>& PipelineCandida
 			PipelineCandidates.AddUnique(Blueprint->GeneratedClass);
 		}
 	}
-}
\ No newline at end of file
+#endif //WITH_ENGINE
+
+}
diff --git a/Engine/Source/Runtime/Interchange/Core/Private/Tasks/InterchangeTaskCompletion.cpp b/Engine/Source/Runtime/Interchange/Core/Private/Tasks/InterchangeTaskCompletion.cpp
--- a/Engine/Source/Runtime/Interchange/Core/Private/Tasks/InterchangeTaskCompletion.cpp
+++ b/Engine/Source/Runtime/Interchange/Core/Private/Tasks/InterchangeTaskCompletion.cpp
@@ -1,7 +1,6 @@
 // Copyright Epic Games, Inc. All Rights Reserved.
 #include "InterchangeTaskCompletion.h"
 
-#include "AssetRegistryModule.h"
 #include "Async/TaskGraphInterfaces.h"
 #include "CoreMinimal.h"
 #include "InterchangeFactoryBase.h"
@@ -12,6 +11,10 @@
 #include "UObject/ObjectMacros.h"
 #include "UObject/WeakObjectPtrTemplates.h"
 
+#if WITH_ENGINE
+#include "AssetRegistryModule.h"
+#endif //WITH_ENGINE
+
 void Interchange::FTaskCompletion::DoTask(ENamedThreads::Type CurrentThread, const FGraphEventRef& MyCompletionGraphEvent)
 {
 	TSharedPtr<Interchange::FImportAsyncHelper> AsyncHelper = WeakAsyncHelper.Pin();
@@ -50,8 +53,10 @@ void Interchange::FTaskCompletion::DoTask(ENamedThreads::Type CurrentThread, con
 			{
 				InterchangeManager->OnAssetPostImport.Broadcast(Asset);
 			}
+#if WITH_ENGINE
 			//Notify the asset registry
 			FAssetRegistryModule::AssetCreated(Asset);
+#endif //WITH_ENGINE
 		}
 	}
 	//Release the async helper
diff --git a/Engine/Source/Runtime/Interchange/Core/Private/Tasks/InterchangeTaskCreateAsset.cpp b/Engine/Source/Runtime/Interchange/Core/Private/Tasks/InterchangeTaskCreateAsset.cpp
--- a/Engine/Source/Runtime/Interchange/Core/Private/Tasks/InterchangeTaskCreateAsset.cpp
+++ b/Engine/Source/Runtime/Interchange/Core/Private/Tasks/InterchangeTaskCreateAsset.cpp
@@ -1,11 +1,10 @@
 // Copyright Epic Games, Inc. All Rights Reserved.
 #include "InterchangeTaskCreateAsset.h"
-#include "InterchangeLogPrivate.h"
 
-#include "AssetRegistryModule.h"
 #include "Async/TaskGraphInterfaces.h"
 #include "CoreMinimal.h"
 #include "InterchangeFactoryBase.h"
+#include "InterchangeLogPrivate.h"
 #include "InterchangeManager.h"
 #include "InterchangeSourceData.h"
 #include "InterchangeTranslatorBase.h"
diff --git a/Engine/Source/Runtime/Interchange/Core/Private/Tasks/InterchangeTaskParsing.cpp b/Engine/Source/Runtime/Interchange/Core/Private/Tasks/InterchangeTaskParsing.cpp
--- a/Engine/Source/Runtime/Interchange/Core/Private/Tasks/InterchangeTaskParsing.cpp
+++ b/Engine/Source/Runtime/Interchange/Core/Private/Tasks/InterchangeTaskParsing.cpp
@@ -1,7 +1,6 @@
 // Copyright Epic Games, Inc. All Rights Reserved.
 #include "InterchangeTaskParsing.h"
 
-#include "AssetRegistryModule.h"
 #include "Async/TaskGraphInterfaces.h"
 #include "CoreMinimal.h"
 #include "InterchangeFactoryBase.h"
diff --git a/Engine/Source/Runtime/Interchange/Core/Public/InterchangeManager.h b/Engine/Source/Runtime/Interchange/Core/Public/InterchangeManager.h
--- a/Engine/Source/Runtime/Interchange/Core/Public/InterchangeManager.h
+++ b/Engine/Source/Runtime/Interchange/Core/Public/InterchangeManager.h
@@ -182,15 +182,6 @@ public:
 	 */
 	bool RegisterFactory(const UClass* Factory);
 
-	/**
-	 * Any c++ pipeline must register to the manager
-	 * @Param pipeline - The UClass of the pipeline you want to register
-	 * @return true if the pipeline class can be register false otherwise.
-	 *
-	 * @Note if you register multiple time the same class it will return true for every call
-	 */
-	bool RegisterPipeline(const UClass* Pipeline);
-
 	/**
 	 * Any writer must register to the manager
 	 * @Param Writer - The UClass of the writer you want to register
diff --git a/Engine/Source/Runtime/Interchange/Core/Public/InterchangeTranslatorBase.h b/Engine/Source/Runtime/Interchange/Core/Public/InterchangeTranslatorBase.h
--- a/Engine/Source/Runtime/Interchange/Core/Public/InterchangeTranslatorBase.h
+++ b/Engine/Source/Runtime/Interchange/Core/Public/InterchangeTranslatorBase.h
@@ -33,4 +33,13 @@ public:
 	{
 		return false;
 	}
+
+	/**
+	 * This function is call when the import is done (FImportAsyncHelper::CleanUp) and we are cleaning the data.
+	 * Use it to free resource that need to be release before the next garbage collector pass.
+	 */
+	virtual void ImportFinish()
+	{
+		return;
+	}
 };
diff --git a/Engine/Source/Runtime/Interchange/Core/Public/Nodes/BaseNode.h b/Engine/Source/Runtime/Interchange/Core/Public/Nodes/BaseNode.h
--- a/Engine/Source/Runtime/Interchange/Core/Public/Nodes/BaseNode.h
+++ b/Engine/Source/Runtime/Interchange/Core/Public/Nodes/BaseNode.h
@@ -79,6 +79,7 @@ namespace InterchangePrivateNodeBase
  * @param AssetType - This is the asset you want to apply the storage value"
  * @param EnumType - Optional, specify it only if the AssetType member is an enum so we can type cast it in the apply function (we use uint8 to store the enum value)"
  */
+#if WITH_ENGINE
 #define IMPLEMENT_NODE_ATTRIBUTE(NodeClassName, AttributeName, AttributeType, AssetType, EnumType)										\
 private:																																\
 const FAttributeKey Macro_Custom##AttributeName##Key = Interchange::FAttributeKey(TEXT(#AttributeName));								\
@@ -116,14 +117,37 @@ bool SetCustom##AttributeName(const AttributeType& AttributeValue, bool bAddAppl
 	{																																	\
 		if(bAddApplyDelegate)																											\
 		{																																\
-			TArray<FApplyAttributeToAsset>& Delegates = ApplyCustomAttributeDelegates.FindOrAdd(AssetType::StaticClass());					\
+			TArray<FApplyAttributeToAsset>& Delegates = ApplyCustomAttributeDelegates.FindOrAdd(AssetType::StaticClass());				\
 			Delegates.Add(FApplyAttributeToAsset::CreateRaw(this, &NodeClassName::ApplyCustom##AttributeName##ToAsset));				\
 		}																																\
 		return true;																													\
 	}																																	\
 	return false;																														\
 }
-
+#else //WITH_ENGINE
+//AssetType is not use when we are not compiling with Engine
+//We then skip the ApplyCustom##AttributeName##ToAsset function
+//Because of this we also skip to add the delegate to ApplyCustomAttributeDelegates
+//Example UTexture is define in engine, a TextureNode compile without the engine will have the attribute
+//setter/getter but will not be able to apply it to a UTexture asset. We use this in out of prcess parser and the graph is reconstruct
+//in the interchange framework which compile with engine so it has the UObject setter call by the interchange factory
+#define IMPLEMENT_NODE_ATTRIBUTE(NodeClassName, AttributeName, AttributeType, AssetType, EnumType)										\
+private:																																\
+const FAttributeKey Macro_Custom##AttributeName##Key = Interchange::FAttributeKey(TEXT(#AttributeName));								\
+																																		\
+public:																																	\
+bool GetCustom##AttributeName(AttributeType& AttributeValue) const																		\
+{																																		\
+	FString OperationName = GetTypeName() + TEXT(".Get##AttributeName");																\
+	return InterchangePrivateNodeBase::GetCustomAttribute<AttributeType>(Attributes, Macro_Custom##AttributeName##Key, OperationName, AttributeValue);	\
+}																																		\
+																																		\
+bool SetCustom##AttributeName(const AttributeType& AttributeValue, bool bAddApplyDelegate = true)										\
+{																																		\
+	FString OperationName = GetTypeName() + TEXT(".Set##AttributeName");																\
+	return InterchangePrivateNodeBase::SetCustomAttribute<AttributeType>(Attributes, Macro_Custom##AttributeName##Key, OperationName, AttributeValue);	\
+}
+#endif //else WITH_ENGINE
 
 
 //Interchange namespace
