Fixed platforms mismatch between Interchange and Datasmith.

diff --git a/Engine/Plugins/Runtime/Interchange/Interchange.uplugin b/Engine/Plugins/Runtime/Interchange/Interchange.uplugin
--- a/Engine/Plugins/Runtime/Interchange/Interchange.uplugin
+++ b/Engine/Plugins/Runtime/Interchange/Interchange.uplugin
@@ -14,61 +14,31 @@
 	"CanContainContent": false,
 	"IsBetaVersion": false,
 	"Installed": false,
-	"SupportedTargetPlatforms": [
-		"Win64",
-		"Linux",
-		"Mac"
-	],
 	"Modules": [
 		{
 			"Name": "InterchangeNodePlugin",
 			"Type": "Runtime",
-			"LoadingPhase": "Default",
-			"WhitelistPlatforms": [
-				"Win64",
-				"Linux",
-				"Mac"
-			]
+			"LoadingPhase": "Default"
 		},
 		{
 			"Name": "InterchangeImportPlugin",
 			"Type": "Runtime",
-			"LoadingPhase": "Default",
-			"WhitelistPlatforms": [
-				"Win64",
-				"Linux",
-				"Mac"
-			]
+			"LoadingPhase": "Default"
 		},
 		{
 			"Name": "InterchangeExportPlugin",
 			"Type": "Runtime",
-			"LoadingPhase": "Default",
-			"WhitelistPlatforms": [
-				"Win64",
-				"Linux",
-				"Mac"
-			]
+			"LoadingPhase": "Default"
 		},
 		{
 			"Name": "InterchangeDispatcher",
 			"Type": "Runtime",
-			"LoadingPhase": "Default",
-			"WhitelistPlatforms": [
-				"Win64",
-				"Linux",
-				"Mac"
-			]
+			"LoadingPhase": "Default"
 		},
 		{
 			"Name": "InterchangeFbxParser",
 			"Type": "Runtime",
-			"LoadingPhase": "Default",
-			"WhitelistPlatforms": [
-				"Win64",
-				"Linux",
-				"Mac"
-			]
+			"LoadingPhase": "Default"
 		}
 	]
 }
\ No newline at end of file
