added missing includes

diff --git a/Engine/Plugins/Runtime/Interchange/Source/Dispatcher/Public/InterchangeDispatcherNetworking.h b/Engine/Plugins/Runtime/Interchange/Source/Dispatcher/Public/InterchangeDispatcherNetworking.h
--- a/Engine/Plugins/Runtime/Interchange/Source/Dispatcher/Public/InterchangeDispatcherNetworking.h
+++ b/Engine/Plugins/Runtime/Interchange/Source/Dispatcher/Public/InterchangeDispatcherNetworking.h
@@ -5,6 +5,7 @@
 #include "InterchangeCommands.h"
 
 #include "CoreMinimal.h"
+#include "Containers/Queue.h"
 #include "Sockets.h"
 
 
diff --git a/Engine/Plugins/Runtime/Interchange/Source/Import/Private/Fbx/InterchangeFbxTranslator.cpp b/Engine/Plugins/Runtime/Interchange/Source/Import/Private/Fbx/InterchangeFbxTranslator.cpp
--- a/Engine/Plugins/Runtime/Interchange/Source/Import/Private/Fbx/InterchangeFbxTranslator.cpp
+++ b/Engine/Plugins/Runtime/Interchange/Source/Import/Private/Fbx/InterchangeFbxTranslator.cpp
@@ -16,6 +16,7 @@
 #include "Nodes/BaseNodeContainer.h"
 #include "Serialization/JsonWriter.h"
 #include "Serialization/JsonSerializer.h"
+#include "HAL/PlatformFileManager.h"
 
 UInterchangeFbxTranslator::UInterchangeFbxTranslator(const class FObjectInitializer& ObjectInitializer)
 {
