- Removed #pragma once from a .cpp file - not portable

diff --git a/Engine/Plugins/Runtime/Interchange/Source/Dispatcher/Private/InterchangeDispatcherTask.cpp b/Engine/Plugins/Runtime/Interchange/Source/Dispatcher/Private/InterchangeDispatcherTask.cpp
--- a/Engine/Plugins/Runtime/Interchange/Source/Dispatcher/Private/InterchangeDispatcherTask.cpp
+++ b/Engine/Plugins/Runtime/Interchange/Source/Dispatcher/Private/InterchangeDispatcherTask.cpp
@@ -1,7 +1,5 @@
 // Copyright Epic Games, Inc. All Rights Reserved.
 
-#pragma once
-
 #include "InterchangeDispatcherTask.h"
 
 #include "CoreMinimal.h"
