Convert FBaseNode to UInterchangeBaseNode and FBaseNodeContainer to UInterchangeBaseNodeContainer. Using UObject will save us the adapter and make all the public API available in c++, python and blueprint.

Add fbx sdk base code to open a file and create the hierarchy so we can import material and textures from fbx

diff --git a/Engine/Plugins/Enterprise/DatasmithImporter/Source/DatasmithImporter/Private/DatasmithTextureImporter.cpp b/Engine/Plugins/Enterprise/DatasmithImporter/Source/DatasmithImporter/Private/DatasmithTextureImporter.cpp
--- a/Engine/Plugins/Enterprise/DatasmithImporter/Source/DatasmithImporter/Private/DatasmithTextureImporter.cpp
+++ b/Engine/Plugins/Enterprise/DatasmithImporter/Source/DatasmithImporter/Private/DatasmithTextureImporter.cpp
@@ -33,7 +33,7 @@
 #include "InterchangeManager.h"
 #include "ObjectTools.h"
 #include "RHI.h"
-#include "TextureNode.h"
+#include "InterchangeTextureNode.h"
 
 #define LOCTEXT_NAMESPACE "DatasmithTextureImport"
 
@@ -327,7 +327,7 @@ Interchange::FAsyncImportResult FDatasmithTextureImporter::CreateTextureAsync(co
 	return FutureTexture;
 }
 
-bool UDatasmithTexturePipeline::ExecuteImportPipeline(UInterchangeBaseNodeContainerAdapter* BaseNodeContainerAdapter)
+bool UDatasmithTexturePipeline::ExecuteImportPipeline(UInterchangeBaseNodeContainer* BaseNodeContainer)
 {
 	if ( !TextureElement.IsValid() )
 	{
@@ -335,14 +335,18 @@ bool UDatasmithTexturePipeline::ExecuteImportPipeline(UInterchangeBaseNodeContai
 	}
 
 	TArray< FName > Nodes;
-	BaseNodeContainerAdapter->GetRoots( Nodes );
+	BaseNodeContainer->GetRoots( Nodes );
 
 	if ( Nodes.Num() <= 0)
 	{
 		return false;
 	}
 
-	Interchange::FTextureNode& TextureNode = static_cast< Interchange::FTextureNode& >( BaseNodeContainerAdapter->BaseNodeContainer->GetNodeChecked( Nodes[0] ) );
+	UInterchangeTextureNode* TextureNode = Cast<UInterchangeTextureNode >( BaseNodeContainer->GetNode( Nodes[0] ) );
+	if (!ensure(TextureNode))
+	{
+		return false;
+	}
 
 	TOptional< bool > bFlipNormalMapGreenChannel;
 	TOptional< TextureMipGenSettings > MipGenSettings;
@@ -418,42 +422,42 @@ bool UDatasmithTexturePipeline::ExecuteImportPipeline(UInterchangeBaseNodeContai
 		bSrgb = false;
 	}
 
-	TextureNode.SetCustomAddressX( (TextureAddress)TextureElement->GetTextureAddressX() );
-	TextureNode.SetCustomAddressY( (TextureAddress)TextureElement->GetTextureAddressY() );
+	TextureNode->SetCustomAddressX( (TextureAddress)TextureElement->GetTextureAddressX() );
+	TextureNode->SetCustomAddressY( (TextureAddress)TextureElement->GetTextureAddressY() );
 
 	if ( bSrgb.IsSet() )
 	{
-		TextureNode.SetCustomSRGB( bSrgb.GetValue() );
+		TextureNode->SetCustomSRGB( bSrgb.GetValue() );
 	}
 
 	if ( bFlipNormalMapGreenChannel.IsSet() )
 	{
-		TextureNode.SetCustombFlipGreenChannel( bFlipNormalMapGreenChannel.GetValue() );
+		TextureNode->SetCustombFlipGreenChannel( bFlipNormalMapGreenChannel.GetValue() );
 	}
 
 	if ( MipGenSettings.IsSet() )
 	{
-		TextureNode.SetCustomMipGenSettings( MipGenSettings.GetValue() );
+		TextureNode->SetCustomMipGenSettings( MipGenSettings.GetValue() );
 	}
 
 	if ( LODGroup.IsSet() )
 	{
-		TextureNode.SetCustomLODGroup( LODGroup.GetValue() );
+		TextureNode->SetCustomLODGroup( LODGroup.GetValue() );
 	}
 
 	if ( CompressionSettings.IsSet() )
 	{
-		TextureNode.SetCustomLODGroup( CompressionSettings.GetValue() );
+		TextureNode->SetCustomLODGroup( CompressionSettings.GetValue() );
 	}
 
 	if ( RGBCurve.IsSet() )
 	{
-		TextureNode.SetCustomAdjustRGBCurve( RGBCurve.GetValue() );
+		TextureNode->SetCustomAdjustRGBCurve( RGBCurve.GetValue() );
 	}
 
 	if ( TexFilter.IsSet() )
 	{
-		TextureNode.SetCustomFilter( TexFilter.GetValue() );
+		TextureNode->SetCustomFilter( TexFilter.GetValue() );
 	}
 
 	return true;
diff --git a/Engine/Plugins/Enterprise/DatasmithImporter/Source/DatasmithImporter/Public/DatasmithTextureImporter.h b/Engine/Plugins/Enterprise/DatasmithImporter/Source/DatasmithImporter/Public/DatasmithTextureImporter.h
--- a/Engine/Plugins/Enterprise/DatasmithImporter/Source/DatasmithImporter/Public/DatasmithTextureImporter.h
+++ b/Engine/Plugins/Enterprise/DatasmithImporter/Source/DatasmithImporter/Public/DatasmithTextureImporter.h
@@ -48,7 +48,7 @@ class UDatasmithTexturePipeline : public UInterchangePipelineBase
 	GENERATED_BODY()
 
 public:
-	virtual bool ExecuteImportPipeline(UInterchangeBaseNodeContainerAdapter* BaseNodeContainerAdapter) override;
+	virtual bool ExecuteImportPipeline(UInterchangeBaseNodeContainer* BaseNodeContainer) override;
 
 	TSharedPtr< IDatasmithTextureElement > TextureElement;
 };
\ No newline at end of file
diff --git a/Engine/Plugins/Runtime/Interchange/Source/Dispatcher/Private/InterchangeDispatcher.cpp b/Engine/Plugins/Runtime/Interchange/Source/Dispatcher/Private/InterchangeDispatcher.cpp
--- a/Engine/Plugins/Runtime/Interchange/Source/Dispatcher/Private/InterchangeDispatcher.cpp
+++ b/Engine/Plugins/Runtime/Interchange/Source/Dispatcher/Private/InterchangeDispatcher.cpp
@@ -46,7 +46,7 @@ TOptional<FTask> FInterchangeDispatcher::GetNextTask()
 	return TaskPool[NextTaskIndex++];
 }
 
-void FInterchangeDispatcher::SetTaskState(int32 TaskIndex, ETaskState TaskState, const FString& JsonResult, const FString& JSonMessages)
+void FInterchangeDispatcher::SetTaskState(int32 TaskIndex, ETaskState TaskState, const FString& JsonResult, const TArray<FString>& JSonMessages)
 {
 	FString JsonDescription;
 	{
@@ -80,7 +80,7 @@ void FInterchangeDispatcher::SetTaskState(int32 TaskIndex, ETaskState TaskState,
 	UE_CLOG(TaskState == ETaskState::ProcessFailed, LogInterchangeDispatcher, Error, TEXT("Json processing failure: %s"), *JsonDescription);
 }
 
-void FInterchangeDispatcher::GetTaskState(int32 TaskIndex, ETaskState& TaskState, FString& JsonResult, FString& JSonMessages)
+void FInterchangeDispatcher::GetTaskState(int32 TaskIndex, ETaskState& TaskState, FString& JsonResult, TArray<FString>& JSonMessages)
 {
 	FScopeLock Lock(&TaskPoolCriticalSection);
 
diff --git a/Engine/Plugins/Runtime/Interchange/Source/Dispatcher/Private/InterchangeWorkerHandler.cpp b/Engine/Plugins/Runtime/Interchange/Source/Dispatcher/Private/InterchangeWorkerHandler.cpp
--- a/Engine/Plugins/Runtime/Interchange/Source/Dispatcher/Private/InterchangeWorkerHandler.cpp
+++ b/Engine/Plugins/Runtime/Interchange/Source/Dispatcher/Private/InterchangeWorkerHandler.cpp
@@ -178,7 +178,8 @@ void FInterchangeWorkerHandler::RunInternal()
 					else
 					{
 						// Signal that the Task was not processed
-						Dispatcher.SetTaskState(CurrentTask->Index, ETaskState::UnTreated, FString(), FString());
+						TArray<FString> GarbageMessages;
+						Dispatcher.SetTaskState(CurrentTask->Index, ETaskState::UnTreated, FString(), GarbageMessages);
 
 						UE_LOG(LogInterchangeDispatcher, Error, TEXT("New task command issue"));
 						WorkerState = EWorkerState::Closing;
@@ -223,7 +224,8 @@ void FInterchangeWorkerHandler::RunInternal()
 					{
 						if (CurrentTask.IsSet())
 						{
-							Dispatcher.SetTaskState(CurrentTask->Index, ETaskState::ProcessFailed, FString(), FString());
+							TArray<FString> GarbageMessages;
+							Dispatcher.SetTaskState(CurrentTask->Index, ETaskState::ProcessFailed, FString(), GarbageMessages);
 							CurrentTask.Reset();
 						}
 
diff --git a/Engine/Plugins/Runtime/Interchange/Source/Dispatcher/Public/InterchangeCommands.h b/Engine/Plugins/Runtime/Interchange/Source/Dispatcher/Public/InterchangeCommands.h
--- a/Engine/Plugins/Runtime/Interchange/Source/Dispatcher/Public/InterchangeCommands.h
+++ b/Engine/Plugins/Runtime/Interchange/Source/Dispatcher/Public/InterchangeCommands.h
@@ -141,7 +141,7 @@ protected:
 public:
 	ETaskState ProcessResult = ETaskState::Unknown;
 	FString JSonResult;
-	FString JSonMessages;
+	TArray<FString> JSonMessages;
 };
 
 } // ns InterchangeDispatcher
diff --git a/Engine/Plugins/Runtime/Interchange/Source/Dispatcher/Public/InterchangeDispatcher.h b/Engine/Plugins/Runtime/Interchange/Source/Dispatcher/Public/InterchangeDispatcher.h
--- a/Engine/Plugins/Runtime/Interchange/Source/Dispatcher/Public/InterchangeDispatcher.h
+++ b/Engine/Plugins/Runtime/Interchange/Source/Dispatcher/Public/InterchangeDispatcher.h
@@ -19,8 +19,8 @@ public:
 
 	int32 AddTask(const FString& JsonDescription);
 	TOptional<FTask> GetNextTask();
-	void SetTaskState(int32 TaskIndex, ETaskState TaskState, const FString& JsonResult, const FString& JSonMessages);
-	void GetTaskState(int32 TaskIndex, ETaskState& TaskState, FString& JsonResult, FString& JSonMessages);
+	void SetTaskState(int32 TaskIndex, ETaskState TaskState, const FString& JsonResult, const TArray<FString>& JSonMessages);
+	void GetTaskState(int32 TaskIndex, ETaskState& TaskState, FString& JsonResult, TArray<FString>& JSonMessages);
 
 	void StartProcess();
 	void TerminateProcess();
diff --git a/Engine/Plugins/Runtime/Interchange/Source/Dispatcher/Public/InterchangeDispatcherTask.h b/Engine/Plugins/Runtime/Interchange/Source/Dispatcher/Public/InterchangeDispatcherTask.h
--- a/Engine/Plugins/Runtime/Interchange/Source/Dispatcher/Public/InterchangeDispatcherTask.h
+++ b/Engine/Plugins/Runtime/Interchange/Source/Dispatcher/Public/InterchangeDispatcherTask.h
@@ -29,7 +29,7 @@ struct FTask
 	int32 Index = -1;
 	ETaskState State = ETaskState::Unknown;
 	FString JsonResult;
-	FString JsonMessages;
+	TArray<FString> JsonMessages;
 };
 
 /**
diff --git a/Engine/Plugins/Runtime/Interchange/Source/Export/Private/InterchangeTextureWriter.cpp b/Engine/Plugins/Runtime/Interchange/Source/Export/Private/InterchangeTextureWriter.cpp
--- a/Engine/Plugins/Runtime/Interchange/Source/Export/Private/InterchangeTextureWriter.cpp
+++ b/Engine/Plugins/Runtime/Interchange/Source/Export/Private/InterchangeTextureWriter.cpp
@@ -1,7 +1,9 @@
 // Copyright Epic Games, Inc. All Rights Reserved. 
 #include "InterchangeTextureWriter.h"
 
-bool UInterchangeTextureWriter::Export(Interchange::FBaseNodeContainer& BaseNodeContainer) const
+#include "Nodes/InterchangeBaseNodeContainer.h"
+
+bool UInterchangeTextureWriter::Export(UInterchangeBaseNodeContainer* BaseNodeContainer) const
 {
 	return false;
 }
diff --git a/Engine/Plugins/Runtime/Interchange/Source/Export/Public/InterchangeTextureWriter.h b/Engine/Plugins/Runtime/Interchange/Source/Export/Public/InterchangeTextureWriter.h
--- a/Engine/Plugins/Runtime/Interchange/Source/Export/Public/InterchangeTextureWriter.h
+++ b/Engine/Plugins/Runtime/Interchange/Source/Export/Public/InterchangeTextureWriter.h
@@ -3,9 +3,10 @@
 #pragma once
 
 #include "CoreMinimal.h"
+#include "InterchangeWriterBase.h"
+#include "Nodes/InterchangeBaseNodeContainer.h"
 #include "UObject/Object.h"
 #include "UObject/ObjectMacros.h"
-#include "InterchangeWriterBase.h"
 
 #include "InterchangeTextureWriter.generated.h"
 
@@ -20,7 +21,7 @@ public:
 	 * @param BaseNodeContainer - Contain nodes describing what to export
 	 * @return true if the writer can export the nodes, false otherwise.
 	 */
-	virtual bool Export(Interchange::FBaseNodeContainer& BaseNodeContainer) const override;
+	virtual bool Export(UInterchangeBaseNodeContainer* BaseNodeContainer) const override;
 };
 
 
diff --git a/Engine/Plugins/Runtime/Interchange/Source/Import/InterchangeImportPlugin.Build.cs b/Engine/Plugins/Runtime/Interchange/Source/Import/InterchangeImportPlugin.Build.cs
--- a/Engine/Plugins/Runtime/Interchange/Source/Import/InterchangeImportPlugin.Build.cs
+++ b/Engine/Plugins/Runtime/Interchange/Source/Import/InterchangeImportPlugin.Build.cs
@@ -15,7 +15,7 @@ namespace UnrealBuildTool.Rules
 					"InterchangeCore",
 					"InterchangeDispatcher",
 					"InterchangeNodePlugin",
-					"MeshDescription"
+					"MeshDescription",
 				}
 				);
 			PrivateDependencyModuleNames.AddRange(
diff --git a/Engine/Plugins/Runtime/Interchange/Source/Import/Private/Fbx/InterchangeFbxTranslator.cpp b/Engine/Plugins/Runtime/Interchange/Source/Import/Private/Fbx/InterchangeFbxTranslator.cpp
--- a/Engine/Plugins/Runtime/Interchange/Source/Import/Private/Fbx/InterchangeFbxTranslator.cpp
+++ b/Engine/Plugins/Runtime/Interchange/Source/Import/Private/Fbx/InterchangeFbxTranslator.cpp
@@ -3,20 +3,23 @@
 
 #include "Dom/JsonValue.h"
 #include "Dom/JsonObject.h"
+#include "HAL/PlatformFileManager.h"
 #include "InterchangeDispatcher.h"
 #include "InterchangeDispatcherTask.h"
+#include "InterchangeManager.h"
 #include "LogInterchangeImportPlugin.h"
-#include "Material/MaterialPayLoad.h"
-#include "Mesh/SkeletalMeshPayload.h"
-#include "Mesh/StaticMeshPayload.h"
+#include "Material/InterchangeMaterialPayloadData.h"
+#include "Mesh/InterchangeSkeletalMeshPayload.h"
+#include "Mesh/InterchangeStaticMeshPayload.h"
 #include "Misc/ConfigCacheIni.h"
 #include "Misc/FileHelper.h"
 #include "Misc/Paths.h"
 #include "Modules/ModuleManager.h"
-#include "Nodes/BaseNodeContainer.h"
+#include "Nodes/InterchangeBaseNodeContainer.h"
 #include "Serialization/JsonWriter.h"
 #include "Serialization/JsonSerializer.h"
-#include "HAL/PlatformFileManager.h"
+#include "Texture/InterchangeTexturePayloadData.h"
+#include "UObject/GCObjectScopeGuard.h"
 
 UInterchangeFbxTranslator::UInterchangeFbxTranslator(const class FObjectInitializer& ObjectInitializer)
 {
@@ -31,7 +34,7 @@ bool UInterchangeFbxTranslator::CanImportSourceData(const UInterchangeSourceData
 	return FbxExtension.StartsWith(Extension,ESearchCase::IgnoreCase);
 }
 
-bool UInterchangeFbxTranslator::Translate(const UInterchangeSourceData* SourceData, Interchange::FBaseNodeContainer& BaseNodeContainer) const
+bool UInterchangeFbxTranslator::Translate(const UInterchangeSourceData* SourceData, UInterchangeBaseNodeContainer& BaseNodeContainer) const
 {
 	FString Filename = SourceData->GetFilename();
 	if (!FPaths::FileExists(Filename))
@@ -86,7 +89,7 @@ bool UInterchangeFbxTranslator::Translate(const UInterchangeSourceData* SourceDa
 		
 	InterchangeDispatcher::ETaskState TaskState;
 	FString JsonResult;
-	FString JSonMessages;
+	TArray<FString> JSonMessages;
 	Dispatcher->GetTaskState(TaskIndex, TaskState, JsonResult, JSonMessages);
 
 	//TODO: Parse the JSonMessage and add the message to the interchange not yet develop error messaging
@@ -101,6 +104,7 @@ bool UInterchangeFbxTranslator::Translate(const UInterchangeSourceData* SourceDa
 	FString BaseNodeContainerFilename = ResultParser.GetResultFilename();
 
 	//Parse the filename to fill the container
+	BaseNodeContainer.LoadFromFile(BaseNodeContainerFilename);
 
 	return true;
 }
@@ -113,40 +117,42 @@ void UInterchangeFbxTranslator::ImportFinish()
 	}
 }
 
-TOptional<Interchange::FStaticMeshPayloadData> UInterchangeFbxTranslator::GetStaticMeshPayloadData(const UInterchangeSourceData* SourceData, const FString& PayLoadKey) const
+
+TOptional<Interchange::FImportImage> UInterchangeFbxTranslator::GetTexturePayloadData(const UInterchangeSourceData* SourceData, const FString& PayLoadKey) const
 {
-	if (!SourceData)
+	UInterchangeSourceData* PayloadSourceData = UInterchangeManager::GetInterchangeManager().CreateSourceData(PayLoadKey);
+	FGCObjectScopeGuard ScopedSourceData(PayloadSourceData);
+	
+	if (!PayloadSourceData)
 	{
-		UE_LOG(LogInterchangeImportPlugin, Error, TEXT("Failed to import Fbx, bad source data."));
-		return TOptional<Interchange::FStaticMeshPayloadData>();
+		return TOptional<Interchange::FImportImage>();
 	}
-	//Use the Source data has the key for the dispatcher and send a command to retrieve the staticmesh bulk data
-	TOptional<Interchange::FStaticMeshPayloadData> PayloadData = Interchange::FStaticMeshPayloadData();
-	return PayloadData;
-}
-
-TOptional<Interchange::FSkeletalMeshPayloadData> UInterchangeFbxTranslator::GetSkeletalMeshPayloadData(const UInterchangeSourceData* SourceData, const FString& PayLoadKey) const
-{
-	if (!SourceData)
+	UInterchangeTranslatorBase* SourceTranslator = UInterchangeManager::GetInterchangeManager().GetTranslatorForSourceData(PayloadSourceData);
+	FGCObjectScopeGuard ScopedSourceTranslator(SourceTranslator);
+	const IInterchangeTexturePayloadInterface* TextureTranslator = Cast<IInterchangeTexturePayloadInterface>(SourceTranslator);
+	if (!ensure(TextureTranslator))
 	{
-		UE_LOG(LogInterchangeImportPlugin, Error, TEXT("Failed to import Fbx, bad source data."));
-		return TOptional<Interchange::FSkeletalMeshPayloadData>();
+		return TOptional<Interchange::FImportImage>();
 	}
-	//Use the Source data has the key for the dispatcher and send a command to retrieve the skeletalmesh bulk data
-	TOptional<Interchange::FSkeletalMeshPayloadData> PayloadData = Interchange::FSkeletalMeshPayloadData();
-	return PayloadData;
+	return TextureTranslator->GetTexturePayloadData(PayloadSourceData, PayLoadKey);
 }
 
 TOptional<Interchange::FMaterialPayloadData> UInterchangeFbxTranslator::GetMaterialPayloadData(const UInterchangeSourceData* SourceData, const FString& PayLoadKey) const
 {
-	if (!SourceData)
-	{
-		UE_LOG(LogInterchangeImportPlugin, Error, TEXT("Failed to import Fbx, bad source data."));
-		return TOptional<Interchange::FMaterialPayloadData>();
-	}
-	//Use the Source data has the key for the dispatcher and send a command to retrieve the material bulk data
-	TOptional<Interchange::FMaterialPayloadData> PayloadData = Interchange::FMaterialPayloadData();
-	return PayloadData;
+	//Not implemented, currently we do not have any payload data for the material
+	return TOptional<Interchange::FMaterialPayloadData>();
+}
+
+TOptional<Interchange::FStaticMeshPayloadData> UInterchangeFbxTranslator::GetStaticMeshPayloadData(const UInterchangeSourceData* SourceData, const FString& PayLoadKey) const
+{
+	//Not implemented, currently we do not have any payload data for static meshes
+	return TOptional<Interchange::FStaticMeshPayloadData>();
+}
+
+TOptional<Interchange::FSkeletalMeshPayloadData> UInterchangeFbxTranslator::GetSkeletalMeshPayloadData(const UInterchangeSourceData* SourceData, const FString& PayLoadKey) const
+{
+	//Not implemented, currently we do not have any payload data for skeletal meshes
+	return TOptional<Interchange::FSkeletalMeshPayloadData>();
 }
 
 FString UInterchangeFbxTranslator::CreateLoadFbxFileCommand(const FString& FbxFilePath) const
diff --git a/Engine/Plugins/Runtime/Interchange/Source/Import/Private/InterchangeImportPlugin.cpp b/Engine/Plugins/Runtime/Interchange/Source/Import/Private/InterchangeImportPlugin.cpp
--- a/Engine/Plugins/Runtime/Interchange/Source/Import/Private/InterchangeImportPlugin.cpp
+++ b/Engine/Plugins/Runtime/Interchange/Source/Import/Private/InterchangeImportPlugin.cpp
@@ -4,6 +4,7 @@
 #include "Fbx/InterchangeFbxTranslator.h"
 #include "IInterchangeImportPlugin.h"
 #include "InterchangeManager.h"
+#include "Material/InterchangeMaterialFactory.h"
 #include "Modules/ModuleManager.h"
 #include "Texture/InterchangeBMPTranslator.h"
 #include "Texture/InterchangeDDSTranslator.h"
@@ -47,6 +48,7 @@ void FInterchangeImportPlugin::StartupModule()
 
 	//Register the factories
 	InterchangeManager.RegisterFactory(UInterchangeTextureFactory::StaticClass());
+	InterchangeManager.RegisterFactory(UInterchangeMaterialFactory::StaticClass());
 }
 
 
diff --git a/Engine/Plugins/Runtime/Interchange/Source/Import/Private/Material/InterchangeMaterialFactory.cpp b/Engine/Plugins/Runtime/Interchange/Source/Import/Private/Material/InterchangeMaterialFactory.cpp
new file mode 100644
--- /dev/null
+++ b/Engine/Plugins/Runtime/Interchange/Source/Import/Private/Material/InterchangeMaterialFactory.cpp
@@ -0,0 +1,296 @@
+// Copyright Epic Games, Inc. All Rights Reserved. 
+#include "Material/InterchangeMaterialFactory.h"
+
+#include "InterchangeMaterialNode.h"
+#include "LogInterchangeImportPlugin.h"
+#include "Materials/Material.h"
+#include "Materials/MaterialExpressionScalarParameter.h"
+#include "Materials/MaterialExpressionTextureCoordinate.h"
+#include "Materials/MaterialExpressionTextureSample.h"
+#include "Materials/MaterialExpressionVectorParameter.h"
+#include "Materials/MaterialInstance.h"
+#include "Materials/MaterialInterface.h"
+#include "Nodes/InterchangeBaseNode.h"
+#include "Nodes/InterchangeBaseNodeContainer.h"
+
+#if WITH_EDITORONLY_DATA
+
+#include "EditorFramework/AssetImportData.h"
+
+#endif //WITH_EDITORONLY_DATA
+
+UClass* UInterchangeMaterialFactory::GetFactoryClass() const
+{
+	return UMaterialInterface::StaticClass();
+}
+
+UObject* UInterchangeMaterialFactory::CreateEmptyAsset(const FCreateAssetParams& Arguments) const
+{
+	UObject* Material = nullptr;
+
+#if WITH_EDITORONLY_DATA
+	if (!Arguments.AssetNode || !Arguments.AssetNode->GetAssetClass()->IsChildOf(GetFactoryClass()))
+	{
+		return nullptr;
+	}
+
+	const UInterchangeMaterialNode* MaterialNode = Cast<UInterchangeMaterialNode>(Arguments.AssetNode);
+	if (MaterialNode == nullptr)
+	{
+		return nullptr;
+	}
+
+	const UClass* MaterialClass = MaterialNode->GetAssetClass();
+	if (!ensure(MaterialClass && MaterialClass->IsChildOf(GetFactoryClass())))
+	{
+		return nullptr;
+	}
+
+	// create an asset if it doesn't exist
+	UObject* ExistingAsset = StaticFindObject(nullptr, Arguments.Parent, *Arguments.AssetName);
+
+	// create a new material or overwrite existing asset, if possible
+	if (!ExistingAsset)
+	{
+		Material = NewObject<UObject>(Arguments.Parent, MaterialClass, *Arguments.AssetName, RF_Public | RF_Standalone);
+	}
+	else if (ExistingAsset->GetClass()->IsChildOf(MaterialClass))
+	{
+		//This is a reimport, we are just re-updating the source data
+		Material = ExistingAsset;
+	}
+
+	if (!Material)
+	{
+		UE_LOG(LogInterchangeImportPlugin, Warning, TEXT("Could not create Material asset %s"), *Arguments.AssetName);
+		return nullptr;
+	}
+#endif //WITH_EDITORONLY_DATA
+	return Material;
+}
+
+UObject* UInterchangeMaterialFactory::CreateAsset(const UInterchangeMaterialFactory::FCreateAssetParams& Arguments) const
+{
+#if !WITH_EDITORONLY_DATA
+
+	UE_LOG(LogInterchangeImportPlugin, Error, TEXT("Cannot import Material asset in runtime, this is an editor only feature."));
+	return nullptr;
+
+#else
+
+	if (!Arguments.AssetNode || !Arguments.AssetNode->GetAssetClass()->IsChildOf(GetFactoryClass()))
+	{
+		return nullptr;
+	}
+
+	const UInterchangeMaterialNode* MaterialNode = Cast<UInterchangeMaterialNode>(Arguments.AssetNode);
+	if (MaterialNode == nullptr)
+	{
+		return nullptr;
+	}
+
+	const UClass* MaterialClass = MaterialNode->GetAssetClass();
+	check(MaterialClass && MaterialClass->IsChildOf(GetFactoryClass()));
+
+	// create an asset if it doesn't exist
+	UObject* ExistingAsset = StaticFindObject(nullptr, Arguments.Parent, *Arguments.AssetName);
+
+	UObject* MaterialObject = nullptr;
+	// create a new material or overwrite existing asset, if possible
+	if (!ExistingAsset)
+	{
+		//NewObject is not thread safe, the asset registry directory watcher tick on the main thread can trig before we finish initializing the UObject and will crash
+		//The UObject should have been create by calling CreateEmptyAsset on the main thread.
+		check(IsInGameThread());
+		MaterialObject = NewObject<UObject>(Arguments.Parent, MaterialClass, *Arguments.AssetName, RF_Public | RF_Standalone);
+	}
+	else if(ExistingAsset->GetClass()->IsChildOf(MaterialClass))
+	{
+		//This is a reimport, we are just re-updating the source data
+		MaterialObject = ExistingAsset;
+	}
+
+	if (!MaterialObject)
+	{
+		UE_LOG(LogInterchangeImportPlugin, Warning, TEXT("Could not create Material asset %s"), *Arguments.AssetName);
+		return nullptr;
+	}
+
+	if (MaterialObject)
+	{
+		//The re-import change only the source data and will not apply any custom attribute to the asset.
+		if(!Arguments.ReimportObject)
+		{
+			UMaterial* Material = Cast<UMaterial>(MaterialObject);
+			
+			if (Material)
+			{
+				auto ApplyInputParameter = [&MaterialNode, &Material, &Arguments](EInterchangeMaterialNodeParameterName ParameterName, FExpressionInput& MaterialInput, const FVector2D& Location)
+				{
+					FName OutTextureUID = NAME_None;
+					int32 OutUVSetIndex = 0;
+					float ScaleU = 0.0f;
+					float ScaleV = 0.0f;
+					FVector OutVectorParameter = FVector(0.0f);
+					float OutFloatParameter = 0.0f;
+					if (MaterialNode->GetTextureParameterData(ParameterName, OutTextureUID, OutUVSetIndex, ScaleU, ScaleV))
+					{
+						UTexture* TextureReference = nullptr;
+						if (Arguments.NodeContainer)
+						{
+							FName NodeUniqueID = MaterialNode->GetUniqueID();
+							int32 ChildCount = Arguments.NodeContainer->GetNodeChildrenCount(NodeUniqueID);
+							for (int32 ChildIndex = 0; ChildIndex < ChildCount; ++ChildIndex)
+							{
+								const UInterchangeBaseNode* ChildNode = Arguments.NodeContainer->GetNodeChildren(NodeUniqueID, ChildIndex);
+								if (ChildNode->GetUniqueID() == OutTextureUID && ChildNode->ReferenceObject.IsAsset())
+								{
+									//Use Resolve object so we just look for in memory UObject
+									UObject* TextureObject = ChildNode->ReferenceObject.ResolveObject();
+									if (TextureObject)
+									{
+										TextureReference = Cast<UTexture>(TextureObject);
+										break;
+									}
+								}
+							}
+						}
+
+						if (!TextureReference)
+						{
+							//TODO: Try to load the texture by extracting an asset name from the filename
+						}
+
+
+						if (TextureReference)
+						{
+							//Set the input to use this texture
+							// and link it to the material 
+							UMaterialExpressionTextureSample* UnrealTextureExpression = NewObject<UMaterialExpressionTextureSample>(Material);
+							Material->Expressions.Add(UnrealTextureExpression);
+							MaterialInput.Expression = UnrealTextureExpression;
+							UnrealTextureExpression->Texture = TextureReference;
+							// 						UnrealTextureExpression->SamplerType = bSetupAsNormalMap ?
+							// 							(bIsVirtualTexture ? SAMPLERTYPE_VirtualNormal : SAMPLERTYPE_Normal) :
+							// 							(bIsVirtualTexture ? SAMPLERTYPE_VirtualColor : SAMPLERTYPE_Color);
+
+							UnrealTextureExpression->MaterialExpressionEditorX = FMath::TruncToInt(Location.X);
+							UnrealTextureExpression->MaterialExpressionEditorY = FMath::TruncToInt(Location.Y);
+
+							//////////////////////////////////////////////////////////////////////////
+							//TODO
+							//This UV set come from fbx we have to find a way to have more context here, the UV index has to be set when creating the node
+							//The name is not suffisant, we need an index
+
+							// add/find UVSet and set it to the texture
+
+							if ((OutUVSetIndex != 0 && OutUVSetIndex != INDEX_NONE) || ScaleU != 1.0f || ScaleV != 1.0f)
+							{
+								// Create a texture coord node for the texture sample
+								UMaterialExpressionTextureCoordinate* MyCoordExpression = NewObject<UMaterialExpressionTextureCoordinate>(Material);
+								Material->Expressions.Add(MyCoordExpression);
+								MyCoordExpression->CoordinateIndex = (OutUVSetIndex >= 0) ? OutUVSetIndex : 0;
+								MyCoordExpression->UTiling = ScaleU;
+								MyCoordExpression->VTiling = ScaleV;
+								UnrealTextureExpression->Coordinates.Expression = MyCoordExpression;
+								MyCoordExpression->MaterialExpressionEditorX = FMath::TruncToInt(Location.X - 175);
+								MyCoordExpression->MaterialExpressionEditorY = FMath::TruncToInt(Location.Y);
+
+							}
+						}
+					}
+					else if (MaterialNode->GetVectorParameterData(ParameterName, OutVectorParameter))
+					{
+						UMaterialExpressionVectorParameter* MyColorExpression = NewObject<UMaterialExpressionVectorParameter>(Material);
+						Material->Expressions.Add( MyColorExpression );
+						MaterialInput.Expression = MyColorExpression;
+
+						MyColorExpression->DefaultValue.R = OutVectorParameter[0];
+						MyColorExpression->DefaultValue.G = OutVectorParameter[1];
+						MyColorExpression->DefaultValue.B = OutVectorParameter[2];
+
+						MyColorExpression->MaterialExpressionEditorX = FMath::TruncToInt(Location.X);
+						MyColorExpression->MaterialExpressionEditorY = FMath::TruncToInt(Location.Y);
+					}
+					else if (MaterialNode->GetScalarParameterData(ParameterName, OutFloatParameter))
+					{
+						UMaterialExpressionScalarParameter* MyScalarExpression = NewObject<UMaterialExpressionScalarParameter>(Material);
+						Material->Expressions.Add(MyScalarExpression);
+						MaterialInput.Expression = MyScalarExpression;
+
+						MyScalarExpression->DefaultValue = OutFloatParameter;
+
+						MyScalarExpression->MaterialExpressionEditorX = FMath::TruncToInt(Location.X);
+						MyScalarExpression->MaterialExpressionEditorY = FMath::TruncToInt(Location.Y);
+					}
+					
+					bool bResult = false;
+					if (MaterialInput.Expression)
+					{
+						TArray<FExpressionOutput> Outputs = MaterialInput.Expression->GetOutputs();
+						FExpressionOutput* Output = Outputs.GetData();
+						MaterialInput.Mask = Output->Mask;
+						MaterialInput.MaskR = Output->MaskR;
+						MaterialInput.MaskG = Output->MaskG;
+						MaterialInput.MaskB = Output->MaskB;
+						MaterialInput.MaskA = Output->MaskA;
+						//If we create an expression we return true
+						bResult = true;
+					}
+					return bResult;
+				};
+
+				ApplyInputParameter(EInterchangeMaterialNodeParameterName::BaseColor, Material->BaseColor, FVector2D(-250, -100));
+				ApplyInputParameter(EInterchangeMaterialNodeParameterName::Metallic, Material->Metallic, FVector2D(-750, 0));
+				ApplyInputParameter(EInterchangeMaterialNodeParameterName::Specular, Material->Specular, FVector2D(-500, 100));
+				ApplyInputParameter(EInterchangeMaterialNodeParameterName::Roughness, Material->Roughness, FVector2D(-250, 200));
+				ApplyInputParameter(EInterchangeMaterialNodeParameterName::EmissiveColor, Material->EmissiveColor, FVector2D(-750, 300));
+				if (ApplyInputParameter(EInterchangeMaterialNodeParameterName::Opacity, Material->Opacity, FVector2D(-500, 400)))
+				{
+					Material->BlendMode = BLEND_Translucent;
+				}
+				ApplyInputParameter(EInterchangeMaterialNodeParameterName::OpacityMask, Material->OpacityMask, FVector2D(-250, 500));
+				ApplyInputParameter(EInterchangeMaterialNodeParameterName::Normal, Material->Normal, FVector2D(-750, 600));
+			}
+
+			/** Apply all MaterialNode custom attributes to the material asset */
+			MaterialNode->ApplyAllCustomAttributeToAsset(Material);
+
+			//TODO support material instance
+		}
+		
+		//Getting the file Hash will cache it into the source data
+		Arguments.SourceData->GetFileContentHash();
+
+		//The interchange completion task (call in the GameThread after the factories pass), will call PostEditChange which will trig another asynchronous system that will build all material in parallel
+	}
+	else
+	{
+		//The material is not a UMaterialInterface
+		MaterialObject->RemoveFromRoot();
+		MaterialObject->MarkPendingKill();
+	}
+	return MaterialObject;
+#endif
+}
+
+/* This function is call in the completion task on the main thread, use it to call main thread post creation step for your assets*/
+void UInterchangeMaterialFactory::PostImportGameThreadCallback(const FPostImportGameThreadCallbackParams& Arguments) const
+{
+	check(IsInGameThread());
+	Super::PostImportGameThreadCallback(Arguments);
+#if WITH_EDITORONLY_DATA
+	if (ensure(Arguments.ReimportObject && Arguments.SourceData))
+	{
+		//We must call the Update of the asset source file in the main thread because UAssetImportData::Update execute some delegate we do not control
+		UMaterialInterface* ImportedMaterial = CastChecked<UMaterialInterface>(Arguments.ReimportObject);
+		//Set the asset import data file source to allow reimport. TODO: manage MD5 Hash properly
+		TOptional<FMD5Hash> FileContentHash = Arguments.SourceData->GetFileContentHash();
+		if (ImportedMaterial->AssetImportData)
+		{
+			ImportedMaterial->AssetImportData->Update(Arguments.SourceData->GetFilename(), FileContentHash.IsSet() ? &FileContentHash.GetValue() : nullptr);
+		}
+	}
+#endif //#if WITH_EDITORONLY_DATA
+	return;
+}
\ No newline at end of file
diff --git a/Engine/Plugins/Runtime/Interchange/Source/Import/Private/Texture/InterchangeBMPTranslator.cpp b/Engine/Plugins/Runtime/Interchange/Source/Import/Private/Texture/InterchangeBMPTranslator.cpp
--- a/Engine/Plugins/Runtime/Interchange/Source/Import/Private/Texture/InterchangeBMPTranslator.cpp
+++ b/Engine/Plugins/Runtime/Interchange/Source/Import/Private/Texture/InterchangeBMPTranslator.cpp
@@ -6,14 +6,15 @@
 #include "Engine/Texture2D.h"
 #include "IImageWrapper.h"
 #include "IImageWrapperModule.h"
+#include "InterchangeTextureNode.h"
 #include "LogInterchangeImportPlugin.h"
 #include "Misc/ConfigCacheIni.h"
 #include "Misc/FileHelper.h"
 #include "Misc/Paths.h"
 #include "Misc/ScopedSlowTask.h"
 #include "Modules/ModuleManager.h"
-#include "Nodes/BaseNodeContainer.h"
-#include "TextureNode.h"
+#include "Nodes/InterchangeBaseNodeContainer.h"
+#include "Texture/TextureTranslatorUtilities.h"
 
 
 
@@ -24,25 +25,12 @@ bool UInterchangeBMPTranslator::CanImportSourceData(const UInterchangeSourceData
 	return BMPExtension.StartsWith(Extension);
 }
 
-bool UInterchangeBMPTranslator::Translate(const UInterchangeSourceData* SourceData, Interchange::FBaseNodeContainer& BaseNodeContainer) const
+bool UInterchangeBMPTranslator::Translate(const UInterchangeSourceData* SourceData, UInterchangeBaseNodeContainer& BaseNodeContainer) const
 {
-	FString Filename = SourceData->GetFilename();
-	if (!FPaths::FileExists(Filename))
-	{
-		return false;
-	}
-
-	FName DisplayLabel = *FPaths::GetBaseFilename(Filename);
-	Interchange::FNodeUniqueID NodeUID(*Filename);
-	// BMP is creating a UTexture2D
-	TUniquePtr<Interchange::FTextureNode> TextureNode = MakeUnique<Interchange::FTextureNode>(NodeUID, DisplayLabel, UTexture2D::StaticClass());
-	TextureNode->SetPayLoadKey(Filename);
-
-	BaseNodeContainer.AddNode(MoveTemp(TextureNode));
-	return true;
+	return Interchange::FTextureTranslatorUtilities::Generic2DTextureTranslate(SourceData, BaseNodeContainer);
 }
 
-const TOptional<Interchange::FImportImage> UInterchangeBMPTranslator::GetPayloadData(const UInterchangeSourceData* SourceData, const FString& PayLoadKey) const
+TOptional<Interchange::FImportImage> UInterchangeBMPTranslator::GetTexturePayloadData(const UInterchangeSourceData* SourceData, const FString& PayLoadKey) const
 {
 	if (!SourceData)
 	{
diff --git a/Engine/Plugins/Runtime/Interchange/Source/Import/Private/Texture/InterchangeDDSTranslator.cpp b/Engine/Plugins/Runtime/Interchange/Source/Import/Private/Texture/InterchangeDDSTranslator.cpp
--- a/Engine/Plugins/Runtime/Interchange/Source/Import/Private/Texture/InterchangeDDSTranslator.cpp
+++ b/Engine/Plugins/Runtime/Interchange/Source/Import/Private/Texture/InterchangeDDSTranslator.cpp
@@ -6,14 +6,15 @@
 #include "Engine/Texture2D.h"
 #include "IImageWrapper.h"
 #include "IImageWrapperModule.h"
+#include "InterchangeTextureNode.h"
 #include "LogInterchangeImportPlugin.h"
 #include "Misc/ConfigCacheIni.h"
 #include "Misc/FileHelper.h"
 #include "Misc/Paths.h"
 #include "Misc/ScopedSlowTask.h"
 #include "Modules/ModuleManager.h"
-#include "Nodes/BaseNodeContainer.h"
-#include "TextureNode.h"
+#include "Nodes/InterchangeBaseNodeContainer.h"
+#include "Texture/TextureTranslatorUtilities.h"
 
 
 //////////////////////////////////////////////////////////////////////////
@@ -58,29 +59,12 @@ bool UInterchangeDDSTranslator::CanImportSourceData(const UInterchangeSourceData
 	return DDSLoadHelper.IsValid2DTexture();
 }
 
-bool UInterchangeDDSTranslator::Translate(const UInterchangeSourceData* SourceData, Interchange::FBaseNodeContainer& BaseNodeContainer) const
+bool UInterchangeDDSTranslator::Translate(const UInterchangeSourceData* SourceData, UInterchangeBaseNodeContainer& BaseNodeContainer) const
 {
-	FString Filename = SourceData->GetFilename();
-	if (!FPaths::FileExists(Filename))
-	{
-		return false;
-	}
-
-	FName DisplayLabel = *FPaths::GetBaseFilename(Filename);
-	Interchange::FNodeUniqueID NodeUID(*Filename);
-	//PCX is creating a UTexture2D
-	TUniquePtr<Interchange::FTextureNode> TextureNode = MakeUnique<Interchange::FTextureNode>(NodeUID, DisplayLabel, UTexture2D::StaticClass());
-	TextureNode->SetPayLoadKey(Filename);
-
-	//Test node change
- 	//TextureNode->SetCustomLODGroup((uint8)TextureGroup::TEXTUREGROUP_WorldNormalMap);
- 	//TextureNode->SetCustomCompressionSettings((uint8)TextureCompressionSettings::TC_Normalmap);
-
-	BaseNodeContainer.AddNode(MoveTemp(TextureNode));
-	return true;
+	return Interchange::FTextureTranslatorUtilities::Generic2DTextureTranslate(SourceData, BaseNodeContainer);
 }
 
-const TOptional<Interchange::FImportImage> UInterchangeDDSTranslator::GetPayloadData(const UInterchangeSourceData* SourceData, const FString& PayLoadKey) const
+TOptional<Interchange::FImportImage> UInterchangeDDSTranslator::GetTexturePayloadData(const UInterchangeSourceData* SourceData, const FString& PayLoadKey) const
 {
 	if (!SourceData)
 	{
diff --git a/Engine/Plugins/Runtime/Interchange/Source/Import/Private/Texture/InterchangeEXRTranslator.cpp b/Engine/Plugins/Runtime/Interchange/Source/Import/Private/Texture/InterchangeEXRTranslator.cpp
--- a/Engine/Plugins/Runtime/Interchange/Source/Import/Private/Texture/InterchangeEXRTranslator.cpp
+++ b/Engine/Plugins/Runtime/Interchange/Source/Import/Private/Texture/InterchangeEXRTranslator.cpp
@@ -5,15 +5,15 @@
 #include "Engine/Texture2D.h"
 #include "IImageWrapper.h"
 #include "IImageWrapperModule.h"
+#include "InterchangeTextureNode.h"
 #include "LogInterchangeImportPlugin.h"
 #include "Misc/ConfigCacheIni.h"
 #include "Misc/FileHelper.h"
 #include "Misc/Paths.h"
 #include "Misc/ScopedSlowTask.h"
 #include "Modules/ModuleManager.h"
-#include "Nodes/BaseNodeContainer.h"
-#include "TextureNode.h"
-
+#include "Nodes/InterchangeBaseNodeContainer.h"
+#include "Texture/TextureTranslatorUtilities.h"
 
 bool UInterchangeEXRTranslator::CanImportSourceData(const UInterchangeSourceData* InSourceData) const
 {
@@ -22,29 +22,12 @@ bool UInterchangeEXRTranslator::CanImportSourceData(const UInterchangeSourceData
 	return EXRExtension.StartsWith(Extension);
 }
 
-bool UInterchangeEXRTranslator::Translate(const UInterchangeSourceData* SourceData, Interchange::FBaseNodeContainer& BaseNodeContainer) const
+bool UInterchangeEXRTranslator::Translate(const UInterchangeSourceData* SourceData, UInterchangeBaseNodeContainer& BaseNodeContainer) const
 {
-	FString Filename = SourceData->GetFilename();
-	if (!FPaths::FileExists(Filename))
-	{
-		return false;
-	}
-
-	FName DisplayLabel = *FPaths::GetBaseFilename(Filename);
-	Interchange::FNodeUniqueID NodeUID(*Filename);
-	//PCX is creating a UTexture2D
-	TUniquePtr<Interchange::FTextureNode> TextureNode = MakeUnique<Interchange::FTextureNode>(NodeUID, DisplayLabel, UTexture2D::StaticClass());
-	TextureNode->SetPayLoadKey(Filename);
-
-	//Test node change
- 	//TextureNode->SetCustomLODGroup((uint8)TextureGroup::TEXTUREGROUP_WorldNormalMap);
- 	//TextureNode->SetCustomCompressionSettings((uint8)TextureCompressionSettings::TC_Normalmap);
-
-	BaseNodeContainer.AddNode(MoveTemp(TextureNode));
-	return true;
+	return Interchange::FTextureTranslatorUtilities::Generic2DTextureTranslate(SourceData, BaseNodeContainer);
 }
 
-const TOptional<Interchange::FImportImage> UInterchangeEXRTranslator::GetPayloadData(const UInterchangeSourceData* SourceData, const FString& PayLoadKey) const
+TOptional<Interchange::FImportImage> UInterchangeEXRTranslator::GetTexturePayloadData(const UInterchangeSourceData* SourceData, const FString& PayLoadKey) const
 {
 	if (!SourceData)
 	{
diff --git a/Engine/Plugins/Runtime/Interchange/Source/Import/Private/Texture/InterchangeJPGTranslator.cpp b/Engine/Plugins/Runtime/Interchange/Source/Import/Private/Texture/InterchangeJPGTranslator.cpp
--- a/Engine/Plugins/Runtime/Interchange/Source/Import/Private/Texture/InterchangeJPGTranslator.cpp
+++ b/Engine/Plugins/Runtime/Interchange/Source/Import/Private/Texture/InterchangeJPGTranslator.cpp
@@ -5,14 +5,15 @@
 #include "Engine/Texture2D.h"
 #include "IImageWrapper.h"
 #include "IImageWrapperModule.h"
+#include "InterchangeTextureNode.h"
 #include "LogInterchangeImportPlugin.h"
 #include "Misc/ConfigCacheIni.h"
 #include "Misc/FileHelper.h"
 #include "Misc/Paths.h"
 #include "Misc/ScopedSlowTask.h"
 #include "Modules/ModuleManager.h"
-#include "Nodes/BaseNodeContainer.h"
-#include "TextureNode.h"
+#include "Nodes/InterchangeBaseNodeContainer.h"
+#include "Texture/TextureTranslatorUtilities.h"
 
 bool UInterchangeJPGTranslator::CanImportSourceData(const UInterchangeSourceData* InSourceData) const
 {
@@ -21,29 +22,12 @@ bool UInterchangeJPGTranslator::CanImportSourceData(const UInterchangeSourceData
 	return JPGExtension.StartsWith(Extension);
 }
 
-bool UInterchangeJPGTranslator::Translate(const UInterchangeSourceData* SourceData, Interchange::FBaseNodeContainer& BaseNodeContainer) const
+bool UInterchangeJPGTranslator::Translate(const UInterchangeSourceData* SourceData, UInterchangeBaseNodeContainer& BaseNodeContainer) const
 {
-	FString Filename = SourceData->GetFilename();
-	if (!FPaths::FileExists(Filename))
-	{
-		return false;
-	}
-
-	FName DisplayLabel = *FPaths::GetBaseFilename(Filename);
-	Interchange::FNodeUniqueID NodeUID(*Filename);
-	//PNG is creating a UTexture2D
-	TUniquePtr<Interchange::FTextureNode> TextureNode = MakeUnique<Interchange::FTextureNode>(NodeUID, DisplayLabel, UTexture2D::StaticClass());
-	TextureNode->SetPayLoadKey(Filename);
-
-	//Test node change
- 	//TextureNode->SetCustomLODGroup((uint8)TextureGroup::TEXTUREGROUP_WorldNormalMap);
- 	//TextureNode->SetCustomCompressionSettings((uint8)TextureCompressionSettings::TC_Normalmap);
-
-	BaseNodeContainer.AddNode(MoveTemp(TextureNode));
-	return true;
+	return Interchange::FTextureTranslatorUtilities::Generic2DTextureTranslate(SourceData, BaseNodeContainer);
 }
 
-const TOptional<Interchange::FImportImage> UInterchangeJPGTranslator::GetPayloadData(const UInterchangeSourceData* SourceData, const FString& PayLoadKey) const
+TOptional<Interchange::FImportImage> UInterchangeJPGTranslator::GetTexturePayloadData(const UInterchangeSourceData* SourceData, const FString& PayLoadKey) const
 {
 	if (!SourceData)
 	{
diff --git a/Engine/Plugins/Runtime/Interchange/Source/Import/Private/Texture/InterchangePCXTranslator.cpp b/Engine/Plugins/Runtime/Interchange/Source/Import/Private/Texture/InterchangePCXTranslator.cpp
--- a/Engine/Plugins/Runtime/Interchange/Source/Import/Private/Texture/InterchangePCXTranslator.cpp
+++ b/Engine/Plugins/Runtime/Interchange/Source/Import/Private/Texture/InterchangePCXTranslator.cpp
@@ -5,15 +5,15 @@
 #include "Engine/Texture2D.h"
 #include "IImageWrapper.h"
 #include "IImageWrapperModule.h"
+#include "InterchangeTextureNode.h"
 #include "LogInterchangeImportPlugin.h"
 #include "Misc/ConfigCacheIni.h"
 #include "Misc/FileHelper.h"
 #include "Misc/Paths.h"
 #include "Misc/ScopedSlowTask.h"
 #include "Modules/ModuleManager.h"
-#include "Nodes/BaseNodeContainer.h"
-#include "TextureNode.h"
-
+#include "Nodes/InterchangeBaseNodeContainer.h"
+#include "Texture/TextureTranslatorUtilities.h"
 
 //////////////////////////////////////////////////////////////////////////
 // PCX helper local function
@@ -68,29 +68,12 @@ bool UInterchangePCXTranslator::CanImportSourceData(const UInterchangeSourceData
 	return PCXExtension.StartsWith(Extension);
 }
 
-bool UInterchangePCXTranslator::Translate(const UInterchangeSourceData* SourceData, Interchange::FBaseNodeContainer& BaseNodeContainer) const
+bool UInterchangePCXTranslator::Translate(const UInterchangeSourceData* SourceData, UInterchangeBaseNodeContainer& BaseNodeContainer) const
 {
-	FString Filename = SourceData->GetFilename();
-	if (!FPaths::FileExists(Filename))
-	{
-		return false;
-	}
-
-	FName DisplayLabel = *FPaths::GetBaseFilename(Filename);
-	Interchange::FNodeUniqueID NodeUID(*Filename);
-	//PCX is creating a UTexture2D
-	TUniquePtr<Interchange::FTextureNode> TextureNode = MakeUnique<Interchange::FTextureNode>(NodeUID, DisplayLabel, UTexture2D::StaticClass());
-	TextureNode->SetPayLoadKey(Filename);
-
-	//Test node change
- 	//TextureNode->SetCustomLODGroup((uint8)TextureGroup::TEXTUREGROUP_WorldNormalMap);
- 	//TextureNode->SetCustomCompressionSettings((uint8)TextureCompressionSettings::TC_Normalmap);
-
-	BaseNodeContainer.AddNode(MoveTemp(TextureNode));
-	return true;
+	return Interchange::FTextureTranslatorUtilities::Generic2DTextureTranslate(SourceData, BaseNodeContainer);
 }
 
-const TOptional<Interchange::FImportImage> UInterchangePCXTranslator::GetPayloadData(const UInterchangeSourceData* SourceData, const FString& PayLoadKey) const
+TOptional<Interchange::FImportImage> UInterchangePCXTranslator::GetTexturePayloadData(const UInterchangeSourceData* SourceData, const FString& PayLoadKey) const
 {
 	if (!SourceData)
 	{
diff --git a/Engine/Plugins/Runtime/Interchange/Source/Import/Private/Texture/InterchangePNGTranslator.cpp b/Engine/Plugins/Runtime/Interchange/Source/Import/Private/Texture/InterchangePNGTranslator.cpp
--- a/Engine/Plugins/Runtime/Interchange/Source/Import/Private/Texture/InterchangePNGTranslator.cpp
+++ b/Engine/Plugins/Runtime/Interchange/Source/Import/Private/Texture/InterchangePNGTranslator.cpp
@@ -5,14 +5,15 @@
 #include "Engine/Texture2D.h"
 #include "IImageWrapper.h"
 #include "IImageWrapperModule.h"
+#include "InterchangeTextureNode.h"
 #include "LogInterchangeImportPlugin.h"
 #include "Misc/ConfigCacheIni.h"
 #include "Misc/FileHelper.h"
 #include "Misc/Paths.h"
 #include "Misc/ScopedSlowTask.h"
 #include "Modules/ModuleManager.h"
-#include "Nodes/BaseNodeContainer.h"
-#include "TextureNode.h"
+#include "Nodes/InterchangeBaseNodeContainer.h"
+#include "Texture/TextureTranslatorUtilities.h"
 
 
 //////////////////////////////////////////////////////////////////////////
@@ -188,29 +189,12 @@ bool UInterchangePNGTranslator::CanImportSourceData(const UInterchangeSourceData
 	return PNGExtension.StartsWith(Extension);
 }
 
-bool UInterchangePNGTranslator::Translate(const UInterchangeSourceData* SourceData, Interchange::FBaseNodeContainer& BaseNodeContainer) const
+bool UInterchangePNGTranslator::Translate(const UInterchangeSourceData* SourceData, UInterchangeBaseNodeContainer& BaseNodeContainer) const
 {
-	FString Filename = SourceData->GetFilename();
-	if (!FPaths::FileExists(Filename))
-	{
-		return false;
-	}
-
-	FName DisplayLabel = *FPaths::GetBaseFilename(Filename);
-	Interchange::FNodeUniqueID NodeUID(*Filename);
-	//PNG is creating a UTexture2D
-	TUniquePtr<Interchange::FTextureNode> TextureNode = MakeUnique<Interchange::FTextureNode>(NodeUID, DisplayLabel, UTexture2D::StaticClass());
-	TextureNode->SetPayLoadKey(Filename);
-
-	//Test node change
- 	//TextureNode->SetCustomLODGroup((uint8)TextureGroup::TEXTUREGROUP_WorldNormalMap);
- 	//TextureNode->SetCustomCompressionSettings((uint8)TextureCompressionSettings::TC_Normalmap);
-
-	BaseNodeContainer.AddNode(MoveTemp(TextureNode));
-	return true;
+	return Interchange::FTextureTranslatorUtilities::Generic2DTextureTranslate(SourceData, BaseNodeContainer);
 }
 
-const TOptional<Interchange::FImportImage> UInterchangePNGTranslator::GetPayloadData(const UInterchangeSourceData* SourceData, const FString& PayLoadKey) const
+TOptional<Interchange::FImportImage> UInterchangePNGTranslator::GetTexturePayloadData(const UInterchangeSourceData* SourceData, const FString& PayLoadKey) const
 {
 	if (!SourceData)
 	{
diff --git a/Engine/Plugins/Runtime/Interchange/Source/Import/Private/Texture/InterchangePSDTranslator.cpp b/Engine/Plugins/Runtime/Interchange/Source/Import/Private/Texture/InterchangePSDTranslator.cpp
--- a/Engine/Plugins/Runtime/Interchange/Source/Import/Private/Texture/InterchangePSDTranslator.cpp
+++ b/Engine/Plugins/Runtime/Interchange/Source/Import/Private/Texture/InterchangePSDTranslator.cpp
@@ -5,15 +5,15 @@
 #include "Engine/Texture2D.h"
 #include "IImageWrapper.h"
 #include "IImageWrapperModule.h"
+#include "InterchangeTextureNode.h"
 #include "LogInterchangeImportPlugin.h"
 #include "Misc/ConfigCacheIni.h"
 #include "Misc/FileHelper.h"
 #include "Misc/Paths.h"
 #include "Misc/ScopedSlowTask.h"
 #include "Modules/ModuleManager.h"
-#include "Nodes/BaseNodeContainer.h"
-#include "TextureNode.h"
-
+#include "Nodes/InterchangeBaseNodeContainer.h"
+#include "Texture/TextureTranslatorUtilities.h"
 
 //////////////////////////////////////////////////////////////////////////
 // PSD helper local function
@@ -370,25 +370,12 @@ bool UInterchangePSDTranslator::CanImportSourceData(const UInterchangeSourceData
 	return PNGExtension.StartsWith(Extension);
 }
 
-bool UInterchangePSDTranslator::Translate(const UInterchangeSourceData* SourceData, Interchange::FBaseNodeContainer& BaseNodeContainer) const
+bool UInterchangePSDTranslator::Translate(const UInterchangeSourceData* SourceData, UInterchangeBaseNodeContainer& BaseNodeContainer) const
 {
-	FString Filename = SourceData->GetFilename();
-	if (!FPaths::FileExists(Filename))
-	{
-		return false;
-	}
-
-	FName DisplayLabel = *FPaths::GetBaseFilename(Filename);
-	Interchange::FNodeUniqueID NodeUID(*Filename);
-	// PSD is creating a UTexture2D
-	TUniquePtr<Interchange::FTextureNode> TextureNode = MakeUnique<Interchange::FTextureNode>(NodeUID, DisplayLabel, UTexture2D::StaticClass());
-	TextureNode->SetPayLoadKey(Filename);
-
-	BaseNodeContainer.AddNode(MoveTemp(TextureNode));
-	return true;
+	return Interchange::FTextureTranslatorUtilities::Generic2DTextureTranslate(SourceData, BaseNodeContainer);
 }
 
-const TOptional<Interchange::FImportImage> UInterchangePSDTranslator::GetPayloadData(const UInterchangeSourceData* SourceData, const FString& PayLoadKey) const
+TOptional<Interchange::FImportImage> UInterchangePSDTranslator::GetTexturePayloadData(const UInterchangeSourceData* SourceData, const FString& PayLoadKey) const
 {
 	if (!SourceData)
 	{
diff --git a/Engine/Plugins/Runtime/Interchange/Source/Import/Private/Texture/InterchangeTGATranslator.cpp b/Engine/Plugins/Runtime/Interchange/Source/Import/Private/Texture/InterchangeTGATranslator.cpp
--- a/Engine/Plugins/Runtime/Interchange/Source/Import/Private/Texture/InterchangeTGATranslator.cpp
+++ b/Engine/Plugins/Runtime/Interchange/Source/Import/Private/Texture/InterchangeTGATranslator.cpp
@@ -3,16 +3,16 @@
 
 #include "IImageWrapper.h"
 #include "IImageWrapperModule.h"
+#include "InterchangeTextureNode.h"
 #include "LogInterchangeImportPlugin.h"
 #include "Misc/ConfigCacheIni.h"
 #include "Misc/FileHelper.h"
 #include "Misc/Paths.h"
 #include "Modules/ModuleManager.h"
-#include "Nodes/BaseNodeContainer.h"
-#include "TextureNode.h"
+#include "Nodes/InterchangeBaseNodeContainer.h"
+#include "Texture/TextureTranslatorUtilities.h"
 #include "TgaImageSupport.h"
 
-
 namespace TGATranslatorImpl
 {
 	bool DecompressTGA( const FTGAFileHeader* TGA, Interchange::FImportImage& OutImage )
@@ -79,26 +79,12 @@ bool UInterchangeTGATranslator::CanImportSourceData(const UInterchangeSourceData
 	return TGAExtension.StartsWith(Extension);
 }
 
-bool UInterchangeTGATranslator::Translate(const UInterchangeSourceData* SourceData, Interchange::FBaseNodeContainer& BaseNodeContainer) const
+bool UInterchangeTGATranslator::Translate(const UInterchangeSourceData* SourceData, UInterchangeBaseNodeContainer& BaseNodeContainer) const
 {
-	FString Filename = SourceData->GetFilename();
-	if ( !FPaths::FileExists( Filename ) )
-	{
-		return false;
-	}
-
-	FName DisplayLabel = *FPaths::GetBaseFilename( Filename );
-	Interchange::FNodeUniqueID NodeUID( *Filename );
-	//TGA is creating a UTexture2D
-	TUniquePtr< Interchange::FTextureNode > TextureNode = MakeUnique< Interchange::FTextureNode >( NodeUID, DisplayLabel, UTexture2D::StaticClass() );
-	TextureNode->SetPayLoadKey( Filename );
-
-	BaseNodeContainer.AddNode( MoveTemp( TextureNode ) );
-
-	return true;
+	return Interchange::FTextureTranslatorUtilities::Generic2DTextureTranslate(SourceData, BaseNodeContainer);
 }
 
-const TOptional<Interchange::FImportImage> UInterchangeTGATranslator::GetPayloadData(const UInterchangeSourceData* SourceData, const FString& PayLoadKey) const
+TOptional<Interchange::FImportImage> UInterchangeTGATranslator::GetTexturePayloadData(const UInterchangeSourceData* SourceData, const FString& PayLoadKey) const
 {
 	if (!SourceData)
 	{
diff --git a/Engine/Plugins/Runtime/Interchange/Source/Import/Private/Texture/InterchangeTextureFactory.cpp b/Engine/Plugins/Runtime/Interchange/Source/Import/Private/Texture/InterchangeTextureFactory.cpp
--- a/Engine/Plugins/Runtime/Interchange/Source/Import/Private/Texture/InterchangeTextureFactory.cpp
+++ b/Engine/Plugins/Runtime/Interchange/Source/Import/Private/Texture/InterchangeTextureFactory.cpp
@@ -1,12 +1,13 @@
 // Copyright Epic Games, Inc. All Rights Reserved. 
 #include "Texture/InterchangeTextureFactory.h"
 
+#include "Engine/Texture.h"
+#include "Engine/Texture2D.h"
+#include "InterchangeTextureNode.h"
 #include "LogInterchangeImportPlugin.h"
-#include "Nodes/BaseNode.h"
-#include "TextureNode.h"
+#include "Nodes/InterchangeBaseNode.h"
+#include "Texture/InterchangeTexturePayloadInterface.h"
 #include "TextureCompiler.h"
-#include "Texture/InterchangeTextureTranslator.h"
-#include "TextureNode.h"
 
 
 #if WITH_EDITORONLY_DATA
@@ -30,7 +31,7 @@ UObject* UInterchangeTextureFactory::CreateEmptyAsset(const FCreateAssetParams&
 		return nullptr;
 	}
 
-	const Interchange::FTextureNode* TextureNode = static_cast<const Interchange::FTextureNode*>(Arguments.AssetNode);
+	const UInterchangeTextureNode* TextureNode = Cast<UInterchangeTextureNode>(Arguments.AssetNode);
 	if (TextureNode == nullptr)
 	{
 		return nullptr;
@@ -46,7 +47,10 @@ UObject* UInterchangeTextureFactory::CreateEmptyAsset(const FCreateAssetParams&
 	}
 
 	const UClass* TextureClass = TextureNode->GetAssetClass();
-	check(TextureClass && TextureClass->IsChildOf(UTexture::StaticClass()));
+	if (!ensure(TextureClass && TextureClass->IsChildOf(GetFactoryClass())))
+	{
+		return nullptr;
+	}
 
 	// create an asset if it doesn't exist
 	UObject* ExistingAsset = StaticFindObject(nullptr, Arguments.Parent, *Arguments.AssetName);
@@ -54,7 +58,7 @@ UObject* UInterchangeTextureFactory::CreateEmptyAsset(const FCreateAssetParams&
 	// create a new texture or overwrite existing asset, if possible
 	if (!ExistingAsset)
 	{
-		Texture = NewObject<UObject>(Arguments.Parent, TextureClass, *Arguments.AssetName, RF_Public);
+		Texture = NewObject<UObject>(Arguments.Parent, TextureClass, *Arguments.AssetName, RF_Public | RF_Standalone);
 	}
 	else if (ExistingAsset->GetClass()->IsChildOf(TextureClass))
 	{
@@ -85,13 +89,13 @@ UObject* UInterchangeTextureFactory::CreateAsset(const UInterchangeTextureFactor
 		return nullptr;
 	}
 
-	const Interchange::FTextureNode* TextureNode = static_cast<const Interchange::FTextureNode*>(Arguments.AssetNode);
+	const UInterchangeTextureNode* TextureNode = Cast<UInterchangeTextureNode>(Arguments.AssetNode);
 	if (TextureNode == nullptr)
 	{
 		return nullptr;
 	}
 
-	const UInterchangeTextureTranslator* TextureTranslator = Cast<UInterchangeTextureTranslator>(Arguments.Translator);
+	const IInterchangeTexturePayloadInterface* TextureTranslator = Cast<IInterchangeTexturePayloadInterface>(Arguments.Translator);
 	if (!TextureTranslator)
 	{
 		return nullptr;
@@ -105,7 +109,7 @@ UObject* UInterchangeTextureFactory::CreateAsset(const UInterchangeTextureFactor
 	{
 		return nullptr;
 	}
-	const TOptional<Interchange::FImportImage> PayloadData = TextureTranslator->GetPayloadData(Arguments.SourceData, PayLoadKey.GetValue());
+	const TOptional<Interchange::FImportImage> PayloadData = TextureTranslator->GetTexturePayloadData(Arguments.SourceData, PayLoadKey.GetValue());
 	if(!PayloadData.IsSet())
 	{
 		return nullptr;
@@ -113,7 +117,7 @@ UObject* UInterchangeTextureFactory::CreateAsset(const UInterchangeTextureFactor
 	const Interchange::FImportImage& Image = PayloadData.GetValue();
 
 	const UClass* TextureClass = TextureNode->GetAssetClass();
-	check(TextureClass && TextureClass->IsChildOf(UTexture::StaticClass()));
+	check(TextureClass && TextureClass->IsChildOf(GetFactoryClass()));
 
 	// create an asset if it doesn't exist
 	UObject* ExistingAsset = StaticFindObject(nullptr, Arguments.Parent, *Arguments.AssetName);
@@ -125,7 +129,7 @@ UObject* UInterchangeTextureFactory::CreateAsset(const UInterchangeTextureFactor
 		//NewObject is not thread safe, the asset registry directory watcher tick on the main thread can trig before we finish initializing the UObject and will crash
 		//The UObject should have been create by calling CreateEmptyAsset on the main thread.
 		check(IsInGameThread());
-		Texture = NewObject<UObject>(Arguments.Parent, TextureClass, *Arguments.AssetName, RF_Public);
+		Texture = NewObject<UObject>(Arguments.Parent, TextureClass, *Arguments.AssetName, RF_Public | RF_Standalone);
 	}
 	else if(ExistingAsset->GetClass()->IsChildOf(TextureClass))
 	{
@@ -161,7 +165,7 @@ UObject* UInterchangeTextureFactory::CreateAsset(const UInterchangeTextureFactor
 		}
 
 		//The re-import change only the source data and will not apply any custom attribute to the asset.
-		if(!ExistingAsset)
+		if(!Arguments.ReimportObject)
 		{
 			/** Apply all TextureNode custom attributes to the texture asset */
 			TextureNode->ApplyAllCustomAttributeToAsset(Texture2D);
diff --git a/Engine/Plugins/Runtime/Interchange/Source/Import/Private/Texture/InterchangeTextureTranslator.cpp b/Engine/Plugins/Runtime/Interchange/Source/Import/Private/Texture/InterchangeTextureTranslator.cpp
deleted file mode 100644
diff --git a/Engine/Plugins/Runtime/Interchange/Source/Import/Private/Texture/TextureTranslatorUtilities.cpp b/Engine/Plugins/Runtime/Interchange/Source/Import/Private/Texture/TextureTranslatorUtilities.cpp
new file mode 100644
--- /dev/null
+++ b/Engine/Plugins/Runtime/Interchange/Source/Import/Private/Texture/TextureTranslatorUtilities.cpp
@@ -0,0 +1,34 @@
+// Copyright Epic Games, Inc. All Rights Reserved. 
+#include "Texture/TextureTranslatorUtilities.h"
+
+#include "CoreMinimal.h"
+#include "InterchangeSourceData.h"
+#include "InterchangeTextureNode.h"
+#include "Misc/Paths.h"
+#include "Nodes/InterchangeBaseNodeContainer.h"
+#include "UObject/Object.h"
+#include "UObject/ObjectMacros.h"
+
+bool Interchange::FTextureTranslatorUtilities::Generic2DTextureTranslate(const UInterchangeSourceData* SourceData, UInterchangeBaseNodeContainer& BaseNodeContainer)
+{
+	FString Filename = SourceData->GetFilename();
+	if (!FPaths::FileExists(Filename))
+	{
+		return false;
+	}
+
+	FName DisplayLabel = *FPaths::GetBaseFilename(Filename);
+	FName NodeUID(*Filename);
+	UInterchangeTextureNode* TextureNode = NewObject<UInterchangeTextureNode>(&BaseNodeContainer, NAME_None);
+	if (!ensure(TextureNode))
+	{
+		return false;
+	}
+	//Creating a UTexture2D
+	TextureNode->InitializeTextureNode(NodeUID, DisplayLabel, UTexture2D::StaticClass()->GetName());
+	TextureNode->SetPayLoadKey(Filename);
+
+	BaseNodeContainer.AddNode(TextureNode);
+
+	return true;
+}
diff --git a/Engine/Plugins/Runtime/Interchange/Source/Import/Private/Texture/TextureTranslatorUtilities.h b/Engine/Plugins/Runtime/Interchange/Source/Import/Private/Texture/TextureTranslatorUtilities.h
new file mode 100644
--- /dev/null
+++ b/Engine/Plugins/Runtime/Interchange/Source/Import/Private/Texture/TextureTranslatorUtilities.h
@@ -0,0 +1,18 @@
+// Copyright Epic Games, Inc. All Rights Reserved.
+#pragma once
+
+#include "CoreMinimal.h"
+#include "Nodes/InterchangeBaseNodeContainer.h"
+#include "InterchangeSourceData.h"
+#include "UObject/Object.h"
+#include "UObject/ObjectMacros.h"
+
+namespace Interchange
+{
+	class FTextureTranslatorUtilities
+	{
+	public:
+		static bool Generic2DTextureTranslate(const UInterchangeSourceData* SourceData, UInterchangeBaseNodeContainer& BaseNodeContainer);
+	};
+}
+
diff --git a/Engine/Plugins/Runtime/Interchange/Source/Import/Public/Fbx/InterchangeFbxTranslator.h b/Engine/Plugins/Runtime/Interchange/Source/Import/Public/Fbx/InterchangeFbxTranslator.h
--- a/Engine/Plugins/Runtime/Interchange/Source/Import/Public/Fbx/InterchangeFbxTranslator.h
+++ b/Engine/Plugins/Runtime/Interchange/Source/Import/Public/Fbx/InterchangeFbxTranslator.h
@@ -1,20 +1,29 @@
 // Copyright Epic Games, Inc. All Rights Reserved.
-
 #pragma once
 
 #include "CoreMinimal.h"
-#include "UObject/Object.h"
-#include "UObject/ObjectMacros.h"
 #include "InterchangeTranslatorBase.h"
 #include "InterchangeDispatcher.h"
-#include "Material/MaterialPayLoad.h"
-#include "Mesh/SkeletalMeshPayload.h"
-#include "Mesh/StaticMeshPayload.h"
+#include "Material/InterchangeMaterialPayloadInterface.h"
+#include "Mesh/InterchangeSkeletalMeshPayload.h"
+#include "Mesh/InterchangeSkeletalMeshPayloadInterface.h"
+#include "Mesh/InterchangeStaticMeshPayload.h"
+#include "Mesh/InterchangeStaticMeshPayloadInterface.h"
+#include "Nodes/InterchangeBaseNodeContainer.h"
+#include "Texture/InterchangeTexturePayloadInterface.h"
+#include "UObject/Object.h"
+#include "UObject/ObjectMacros.h"
 
 #include "InterchangeFbxTranslator.generated.h"
 
+/* Fbx translator class support import of texture, material, static mesh, skeletal mesh, */
+
 UCLASS(BlueprintType)
 class INTERCHANGEIMPORTPLUGIN_API UInterchangeFbxTranslator : public UInterchangeTranslatorBase
+, public IInterchangeTexturePayloadInterface
+, public IInterchangeMaterialPayloadInterface
+, public IInterchangeStaticMeshPayloadInterface
+, public IInterchangeSkeletalMeshPayloadInterface
 {
 	GENERATED_BODY()
 public:
@@ -22,19 +31,29 @@ public:
 
 	/** Begin UInterchangeTranslatorBase API*/
 	virtual bool CanImportSourceData(const UInterchangeSourceData* InSourceData) const override;
-	virtual bool Translate(const UInterchangeSourceData* SourceData, Interchange::FBaseNodeContainer& BaseNodeContainer) const override;
+	virtual bool Translate(const UInterchangeSourceData* SourceData, UInterchangeBaseNodeContainer& BaseNodeContainer) const override;
 	virtual void ImportFinish() override;
 	/** End UInterchangeTranslatorBase API*/
 
+
+	//////////////////////////////////////////////////////////////////////////
+	/* IInterchangeTexturePayloadInterface Begin */
+
 	/**
 	 * Once the translation is done, the import process need a way to retrieve payload data.
 	 * This payload will be use by the factories to create the asset.
 	 *
 	 * @param SourceData - The source data containing the data to translate
 	 * @param PayloadKey - The key to retrieve the a particular payload contain into the specified source data.
-	 * @return a PayloadData containing the mesh data to import. The TOptional will not be set if there is an error.
+	 * @return a PayloadData containing the imported data. The TOptional will not be set if there is an error.
 	 */
-	virtual TOptional<Interchange::FStaticMeshPayloadData> GetStaticMeshPayloadData(const UInterchangeSourceData* SourceData, const FString& PayLoadKey) const;
+	virtual TOptional<Interchange::FImportImage> GetTexturePayloadData(const UInterchangeSourceData* SourceData, const FString& PayLoadKey) const override;
+
+	/* IInterchangeTexturePayloadInterface End */
+
+
+	//////////////////////////////////////////////////////////////////////////
+	/* IInterchangeMaterialPayloadInterface Begin */
 
 	/**
 	 * Once the translation is done, the import process need a way to retrieve payload data.
@@ -42,9 +61,31 @@ public:
 	 *
 	 * @param SourceData - The source data containing the data to translate
 	 * @param PayloadKey - The key to retrieve the a particular payload contain into the specified source data.
-	 * @return a PayloadData containing the mesh data to import. The TOptional will not be set if there is an error.
+	 * @return a PayloadData containing the imported data. The TOptional will not be set if there is an error.
 	 */
-	virtual TOptional<Interchange::FSkeletalMeshPayloadData> GetSkeletalMeshPayloadData(const UInterchangeSourceData* SourceData, const FString& PayLoadKey) const;
+	virtual TOptional<Interchange::FMaterialPayloadData> GetMaterialPayloadData(const UInterchangeSourceData* SourceData, const FString& PayLoadKey) const override;
+
+	/* IInterchangeMaterialPayloadInterface End */
+
+
+	//////////////////////////////////////////////////////////////////////////
+	/* IInterchangeStaticMeshPayloadInterface Begin */
+
+	/**
+	 * Once the translation is done, the import process need a way to retrieve payload data.
+	 * This payload will be use by the factories to create the asset.
+	 *
+	 * @param SourceData - The source data containing the data to translate
+	 * @param PayloadKey - The key to retrieve the a particular payload contain into the specified source data.
+	 * @return a PayloadData containing the imported data. The TOptional will not be set if there is an error.
+	 */
+	virtual TOptional<Interchange::FStaticMeshPayloadData> GetStaticMeshPayloadData(const UInterchangeSourceData* SourceData, const FString& PayLoadKey) const override;
+
+	/* IInterchangeStaticMeshPayloadInterface End */
+
+
+	//////////////////////////////////////////////////////////////////////////
+	/* IInterchangeSkeletalMeshPayloadInterface Begin */
 
 	/**
 	 * Once the translation is done, the import process need a way to retrieve payload data.
@@ -52,10 +93,11 @@ public:
 	 *
 	 * @param SourceData - The source data containing the data to translate
 	 * @param PayloadKey - The key to retrieve the a particular payload contain into the specified source data.
-	 * @return a PayloadData containing the mesh data to import. The TOptional will not be set if there is an error.
+	 * @return a PayloadData containing the imported data. The TOptional will not be set if there is an error.
 	 */
-	virtual TOptional<Interchange::FMaterialPayloadData> GetMaterialPayloadData(const UInterchangeSourceData* SourceData, const FString& PayLoadKey) const;
+	virtual TOptional<Interchange::FSkeletalMeshPayloadData> GetSkeletalMeshPayloadData(const UInterchangeSourceData* SourceData, const FString& PayLoadKey) const override;
 
+	/* IInterchangeSkeletalMeshPayloadInterface End */
 private:
 
 	FString CreateLoadFbxFileCommand(const FString& FbxFilePath) const;
diff --git a/Engine/Plugins/Runtime/Interchange/Source/Import/Public/Material/InterchangeMaterialFactory.h b/Engine/Plugins/Runtime/Interchange/Source/Import/Public/Material/InterchangeMaterialFactory.h
new file mode 100644
--- /dev/null
+++ b/Engine/Plugins/Runtime/Interchange/Source/Import/Public/Material/InterchangeMaterialFactory.h
@@ -0,0 +1,30 @@
+// Copyright Epic Games, Inc. All Rights Reserved.
+
+#pragma once
+
+#include "CoreMinimal.h"
+#include "InterchangeFactoryBase.h"
+#include "UObject/Object.h"
+#include "UObject/ObjectMacros.h"
+
+#include "InterchangeMaterialFactory.generated.h"
+
+UCLASS(BlueprintType)
+class INTERCHANGEIMPORTPLUGIN_API UInterchangeMaterialFactory : public UInterchangeFactoryBase
+{
+	GENERATED_BODY()
+public:
+
+	//////////////////////////////////////////////////////////////////////////
+	// Interchange factory base interface begin
+
+	virtual UClass* GetFactoryClass() const override;
+	virtual UObject* CreateEmptyAsset(const FCreateAssetParams& Arguments) const override;
+	virtual UObject* CreateAsset(const FCreateAssetParams& Arguments) const override;
+	virtual void PostImportGameThreadCallback(const FPostImportGameThreadCallbackParams& Arguments) const override;
+
+	// Interchange factory base interface end
+	//////////////////////////////////////////////////////////////////////////
+};
+
+
diff --git a/Engine/Plugins/Runtime/Interchange/Source/Import/Public/Material/InterchangeMaterialPayloadData.h b/Engine/Plugins/Runtime/Interchange/Source/Import/Public/Material/InterchangeMaterialPayloadData.h
new file mode 100644
--- /dev/null
+++ b/Engine/Plugins/Runtime/Interchange/Source/Import/Public/Material/InterchangeMaterialPayloadData.h
@@ -0,0 +1,15 @@
+// Copyright Epic Games, Inc. All Rights Reserved.
+
+#pragma once
+
+#include "CoreMinimal.h"
+
+namespace Interchange
+{
+	struct FMaterialPayloadData
+	{
+		//Store any material payload data here, currently not used
+	};
+}
+
+
diff --git a/Engine/Plugins/Runtime/Interchange/Source/Import/Public/Material/InterchangeMaterialPayloadInterface.h b/Engine/Plugins/Runtime/Interchange/Source/Import/Public/Material/InterchangeMaterialPayloadInterface.h
new file mode 100644
--- /dev/null
+++ b/Engine/Plugins/Runtime/Interchange/Source/Import/Public/Material/InterchangeMaterialPayloadInterface.h
@@ -0,0 +1,36 @@
+// Copyright Epic Games, Inc. All Rights Reserved.
+#pragma once
+
+#include "CoreMinimal.h"
+#include "UObject/Object.h"
+#include "UObject/ObjectMacros.h"
+#include "Material/InterchangeMaterialPayloadData.h"
+
+#include "InterchangeMaterialPayloadInterface.generated.h"
+
+UINTERFACE()
+class INTERCHANGEIMPORTPLUGIN_API UInterchangeMaterialPayloadInterface : public UInterface
+{
+	GENERATED_BODY()
+};
+
+/**
+ * Material translator interface. Derive from it if your translator can import materials
+ */
+class INTERCHANGEIMPORTPLUGIN_API IInterchangeMaterialPayloadInterface
+{
+	GENERATED_BODY()
+public:
+
+	/**
+	 * Once the translation is done, the import process need a way to retrieve payload data.
+	 * This payload will be use by the factories to create the asset.
+	 *
+	 * @param SourceData - The source data containing the data to translate
+	 * @param PayloadKey - The key to retrieve the a particular payload contain into the specified source data.
+	 * @return a PayloadData containing the import image data. The TOptional will not be set if there is an error.
+	 */
+	virtual TOptional<Interchange::FMaterialPayloadData> GetMaterialPayloadData(const UInterchangeSourceData* SourceData, const FString& PayloadKey) const = 0;
+};
+
+
diff --git a/Engine/Plugins/Runtime/Interchange/Source/Import/Public/Material/MaterialPayLoad.h b/Engine/Plugins/Runtime/Interchange/Source/Import/Public/Material/MaterialPayLoad.h
deleted file mode 100644
diff --git a/Engine/Plugins/Runtime/Interchange/Source/Import/Public/Mesh/InterchangeSkeletalMeshPayload.h b/Engine/Plugins/Runtime/Interchange/Source/Import/Public/Mesh/InterchangeSkeletalMeshPayload.h
new file mode 100644
--- /dev/null
+++ b/Engine/Plugins/Runtime/Interchange/Source/Import/Public/Mesh/InterchangeSkeletalMeshPayload.h
@@ -0,0 +1,17 @@
+// Copyright Epic Games, Inc. All Rights Reserved.
+
+#pragma once
+
+#include "CoreMinimal.h"
+#include "Rendering/SkeletalMeshLODImporterData.h"
+
+namespace Interchange
+{
+	struct FSkeletalMeshPayloadData
+	{
+		//Currently the skeletalmesh payload data is editor only, we have to move to something available at runtime
+#if WITH_EDITOR
+		TArray<FSkeletalMeshImportData> SkeletalMeshImportData;
+#endif
+	};
+}
diff --git a/Engine/Plugins/Runtime/Interchange/Source/Import/Public/Mesh/InterchangeSkeletalMeshPayloadInterface.h b/Engine/Plugins/Runtime/Interchange/Source/Import/Public/Mesh/InterchangeSkeletalMeshPayloadInterface.h
new file mode 100644
--- /dev/null
+++ b/Engine/Plugins/Runtime/Interchange/Source/Import/Public/Mesh/InterchangeSkeletalMeshPayloadInterface.h
@@ -0,0 +1,36 @@
+// Copyright Epic Games, Inc. All Rights Reserved.
+#pragma once
+
+#include "CoreMinimal.h"
+#include "Mesh/InterchangeSkeletalMeshPayload.h"
+#include "UObject/Object.h"
+#include "UObject/ObjectMacros.h"
+
+#include "InterchangeSkeletalMeshPayloadInterface.generated.h"
+
+UINTERFACE()
+class INTERCHANGEIMPORTPLUGIN_API UInterchangeSkeletalMeshPayloadInterface : public UInterface
+{
+	GENERATED_BODY()
+};
+
+/**
+ * Skeletal mesh payload interface. Derive from this interface if your payload can import skeletal mesh
+ */
+class INTERCHANGEIMPORTPLUGIN_API IInterchangeSkeletalMeshPayloadInterface
+{
+	GENERATED_BODY()
+public:
+
+	/**
+	 * Once the translation is done, the import process need a way to retrieve payload data.
+	 * This payload will be use by the factories to create the asset.
+	 *
+	 * @param SourceData - The source data containing the data to translate
+	 * @param PayloadKey - The key to retrieve the a particular payload contain into the specified source data.
+	 * @return a PayloadData containing the import image data. The TOptional will not be set if there is an error.
+	 */
+	virtual TOptional<Interchange::FSkeletalMeshPayloadData> GetSkeletalMeshPayloadData(const UInterchangeSourceData* SourceData, const FString& PayloadKey) const = 0;
+};
+
+
diff --git a/Engine/Plugins/Runtime/Interchange/Source/Import/Public/Mesh/InterchangeStaticMeshPayload.h b/Engine/Plugins/Runtime/Interchange/Source/Import/Public/Mesh/InterchangeStaticMeshPayload.h
new file mode 100644
--- /dev/null
+++ b/Engine/Plugins/Runtime/Interchange/Source/Import/Public/Mesh/InterchangeStaticMeshPayload.h
@@ -0,0 +1,14 @@
+// Copyright Epic Games, Inc. All Rights Reserved.
+
+#pragma once
+
+#include "CoreMinimal.h"
+#include "MeshDescription.h"
+
+namespace Interchange
+{
+	struct FStaticMeshPayloadData
+	{
+		TArray<FMeshDescription> MeshDescriptions;
+	};
+}
diff --git a/Engine/Plugins/Runtime/Interchange/Source/Import/Public/Mesh/InterchangeStaticMeshPayloadInterface.h b/Engine/Plugins/Runtime/Interchange/Source/Import/Public/Mesh/InterchangeStaticMeshPayloadInterface.h
new file mode 100644
--- /dev/null
+++ b/Engine/Plugins/Runtime/Interchange/Source/Import/Public/Mesh/InterchangeStaticMeshPayloadInterface.h
@@ -0,0 +1,36 @@
+// Copyright Epic Games, Inc. All Rights Reserved.
+#pragma once
+
+#include "CoreMinimal.h"
+#include "UObject/Object.h"
+#include "UObject/ObjectMacros.h"
+#include "Mesh/InterchangeStaticMeshPayload.h"
+
+#include "InterchangeStaticMeshPayloadInterface.generated.h"
+
+UINTERFACE()
+class INTERCHANGEIMPORTPLUGIN_API UInterchangeStaticMeshPayloadInterface : public UInterface
+{
+	GENERATED_BODY()
+};
+
+/**
+ * Static mesh payload interface. Derive from this interface if your payload can import static mesh
+ */
+class INTERCHANGEIMPORTPLUGIN_API IInterchangeStaticMeshPayloadInterface
+{
+	GENERATED_BODY()
+public:
+
+	/**
+	 * Once the translation is done, the import process need a way to retrieve payload data.
+	 * This payload will be use by the factories to create the asset.
+	 *
+	 * @param SourceData - The source data containing the data to translate
+	 * @param PayloadKey - The key to retrieve the a particular payload contain into the specified source data.
+	 * @return a PayloadData containing the import image data. The TOptional will not be set if there is an error.
+	 */
+	virtual TOptional<Interchange::FStaticMeshPayloadData> GetStaticMeshPayloadData(const UInterchangeSourceData* SourceData, const FString& PayloadKey) const = 0;
+};
+
+
diff --git a/Engine/Plugins/Runtime/Interchange/Source/Import/Public/Mesh/SkeletalMeshPayload.h b/Engine/Plugins/Runtime/Interchange/Source/Import/Public/Mesh/SkeletalMeshPayload.h
deleted file mode 100644
diff --git a/Engine/Plugins/Runtime/Interchange/Source/Import/Public/Mesh/StaticMeshPayload.h b/Engine/Plugins/Runtime/Interchange/Source/Import/Public/Mesh/StaticMeshPayload.h
deleted file mode 100644
diff --git a/Engine/Plugins/Runtime/Interchange/Source/Import/Public/Texture/InterchangeBMPTranslator.h b/Engine/Plugins/Runtime/Interchange/Source/Import/Public/Texture/InterchangeBMPTranslator.h
--- a/Engine/Plugins/Runtime/Interchange/Source/Import/Public/Texture/InterchangeBMPTranslator.h
+++ b/Engine/Plugins/Runtime/Interchange/Source/Import/Public/Texture/InterchangeBMPTranslator.h
@@ -6,13 +6,14 @@
 #include "UObject/Object.h"
 #include "UObject/ObjectMacros.h"
 #include "InterchangeTranslatorBase.h"
+#include "Nodes/InterchangeBaseNodeContainer.h"
 #include "Texture/InterchangeTexturePayloadData.h"
-#include "Texture/InterchangeTextureTranslator.h"
+#include "Texture/InterchangeTexturePayloadInterface.h"
 
 #include "InterchangeBMPTranslator.generated.h"
 
 UCLASS(BlueprintType)
-class INTERCHANGEIMPORTPLUGIN_API UInterchangeBMPTranslator : public UInterchangeTextureTranslator
+class INTERCHANGEIMPORTPLUGIN_API UInterchangeBMPTranslator : public UInterchangeTranslatorBase, public IInterchangeTexturePayloadInterface
 {
 	GENERATED_BODY()
 public:
@@ -29,7 +30,10 @@ public:
 	 * @param BaseNodeContainer - The unreal objects descriptions container where to put the translated source data.
 	 * @return true if the translator can translate the source data, false otherwise.
 	 */
-	virtual bool Translate(const UInterchangeSourceData* SourceData, Interchange::FBaseNodeContainer& BaseNodeContainer) const override;
+	virtual bool Translate(const UInterchangeSourceData* SourceData, UInterchangeBaseNodeContainer& BaseNodeContainer) const override;
+
+
+	/* IInterchangeTexturePayloadInterface Begin */
 
 	/**
 	 * Once the translation is done, the import process need a way to retrieve payload data.
@@ -39,5 +43,7 @@ public:
 	 * @param PayloadKey - The key to retrieve the a particular payload contain into the specified source data.
 	 * @return a PayloadData containing the import image data. The TOptional will not be set if there is an error.
 	 */
-	virtual const TOptional<Interchange::FImportImage> GetPayloadData(const UInterchangeSourceData* SourceData, const FString& PayLoadKey) const override;
+	virtual TOptional<Interchange::FImportImage> GetTexturePayloadData(const UInterchangeSourceData* SourceData, const FString& PayLoadKey) const override;
+
+	/* IInterchangeTexturePayloadInterface End */
 };
diff --git a/Engine/Plugins/Runtime/Interchange/Source/Import/Public/Texture/InterchangeDDSTranslator.h b/Engine/Plugins/Runtime/Interchange/Source/Import/Public/Texture/InterchangeDDSTranslator.h
--- a/Engine/Plugins/Runtime/Interchange/Source/Import/Public/Texture/InterchangeDDSTranslator.h
+++ b/Engine/Plugins/Runtime/Interchange/Source/Import/Public/Texture/InterchangeDDSTranslator.h
@@ -6,13 +6,14 @@
 #include "UObject/Object.h"
 #include "UObject/ObjectMacros.h"
 #include "InterchangeTranslatorBase.h"
+#include "Nodes/InterchangeBaseNodeContainer.h"
 #include "Texture/InterchangeTexturePayloadData.h"
-#include "Texture/InterchangeTextureTranslator.h"
+#include "Texture/InterchangeTexturePayloadInterface.h"
 
 #include "InterchangeDDSTranslator.generated.h"
 
 UCLASS(BlueprintType)
-class INTERCHANGEIMPORTPLUGIN_API UInterchangeDDSTranslator : public UInterchangeTextureTranslator
+class INTERCHANGEIMPORTPLUGIN_API UInterchangeDDSTranslator : public UInterchangeTranslatorBase, public IInterchangeTexturePayloadInterface
 {
 	GENERATED_BODY()
 public:
@@ -29,7 +30,9 @@ public:
 	 * @param BaseNodeContainer - The unreal objects descriptions container where to put the translated source data.
 	 * @return true if the translator can translate the source data, false otherwise.
 	 */
-	virtual bool Translate(const UInterchangeSourceData* SourceData, Interchange::FBaseNodeContainer& BaseNodeContainer) const override;
+	virtual bool Translate(const UInterchangeSourceData* SourceData, UInterchangeBaseNodeContainer& BaseNodeContainer) const override;
+
+	/* IInterchangeTexturePayloadInterface Begin */
 
 	/**
 	 * Once the translation is done, the import process need a way to retrieve payload data.
@@ -39,7 +42,9 @@ public:
 	 * @param PayloadKey - The key to retrieve the a particular payload contain into the specified source data.
 	 * @return a PayloadData containing the import image data. The TOptional will not be set if there is an error.
 	 */
-	virtual const TOptional<Interchange::FImportImage> GetPayloadData(const UInterchangeSourceData* SourceData, const FString& PayLoadKey) const override;
+	virtual TOptional<Interchange::FImportImage> GetTexturePayloadData(const UInterchangeSourceData* SourceData, const FString& PayLoadKey) const override;
+
+	/* IInterchangeTexturePayloadInterface End */
 };
 
 
diff --git a/Engine/Plugins/Runtime/Interchange/Source/Import/Public/Texture/InterchangeEXRTranslator.h b/Engine/Plugins/Runtime/Interchange/Source/Import/Public/Texture/InterchangeEXRTranslator.h
--- a/Engine/Plugins/Runtime/Interchange/Source/Import/Public/Texture/InterchangeEXRTranslator.h
+++ b/Engine/Plugins/Runtime/Interchange/Source/Import/Public/Texture/InterchangeEXRTranslator.h
@@ -6,13 +6,14 @@
 #include "UObject/Object.h"
 #include "UObject/ObjectMacros.h"
 #include "InterchangeTranslatorBase.h"
+#include "Nodes/InterchangeBaseNodeContainer.h"
 #include "Texture/InterchangeTexturePayloadData.h"
-#include "Texture/InterchangeTextureTranslator.h"
+#include "Texture/InterchangeTexturePayloadInterface.h"
 
 #include "InterchangeEXRTranslator.generated.h"
 
 UCLASS(BlueprintType)
-class INTERCHANGEIMPORTPLUGIN_API UInterchangeEXRTranslator : public UInterchangeTextureTranslator
+class INTERCHANGEIMPORTPLUGIN_API UInterchangeEXRTranslator : public UInterchangeTranslatorBase, public IInterchangeTexturePayloadInterface
 {
 	GENERATED_BODY()
 public:
@@ -29,7 +30,9 @@ public:
 	 * @param BaseNodeContainer - The unreal objects descriptions container where to put the translated source data.
 	 * @return true if the translator can translate the source data, false otherwise.
 	 */
-	virtual bool Translate(const UInterchangeSourceData* SourceData, Interchange::FBaseNodeContainer& BaseNodeContainer) const override;
+	virtual bool Translate(const UInterchangeSourceData* SourceData, UInterchangeBaseNodeContainer& BaseNodeContainer) const override;
+
+	/* IInterchangeTexturePayloadInterface Begin */
 
 	/**
 	 * Once the translation is done, the import process need a way to retrieve payload data.
@@ -39,7 +42,9 @@ public:
 	 * @param PayloadKey - The key to retrieve the a particular payload contain into the specified source data.
 	 * @return a PayloadData containing the import image data. The TOptional will not be set if there is an error.
 	 */
-	virtual const TOptional<Interchange::FImportImage> GetPayloadData(const UInterchangeSourceData* SourceData, const FString& PayLoadKey) const override;
+	virtual TOptional<Interchange::FImportImage> GetTexturePayloadData(const UInterchangeSourceData* SourceData, const FString& PayLoadKey) const override;
+
+	/* IInterchangeTexturePayloadInterface End */
 };
 
 
diff --git a/Engine/Plugins/Runtime/Interchange/Source/Import/Public/Texture/InterchangeJPGTranslator.h b/Engine/Plugins/Runtime/Interchange/Source/Import/Public/Texture/InterchangeJPGTranslator.h
--- a/Engine/Plugins/Runtime/Interchange/Source/Import/Public/Texture/InterchangeJPGTranslator.h
+++ b/Engine/Plugins/Runtime/Interchange/Source/Import/Public/Texture/InterchangeJPGTranslator.h
@@ -6,13 +6,14 @@
 #include "UObject/Object.h"
 #include "UObject/ObjectMacros.h"
 #include "InterchangeTranslatorBase.h"
+#include "Nodes/InterchangeBaseNodeContainer.h"
 #include "Texture/InterchangeTexturePayloadData.h"
-#include "Texture/InterchangeTextureTranslator.h"
+#include "Texture/InterchangeTexturePayloadInterface.h"
 
 #include "InterchangeJPGTranslator.generated.h"
 
 UCLASS(BlueprintType)
-class INTERCHANGEIMPORTPLUGIN_API UInterchangeJPGTranslator : public UInterchangeTextureTranslator
+class INTERCHANGEIMPORTPLUGIN_API UInterchangeJPGTranslator : public UInterchangeTranslatorBase, public IInterchangeTexturePayloadInterface
 {
 	GENERATED_BODY()
 public:
@@ -29,7 +30,9 @@ public:
 	 * @param BaseNodeContainer - The unreal objects descriptions container where to put the translated source data.
 	 * @return true if the translator can translate the source data, false otherwise.
 	 */
-	virtual bool Translate(const UInterchangeSourceData* SourceData, Interchange::FBaseNodeContainer& BaseNodeContainer) const override;
+	virtual bool Translate(const UInterchangeSourceData* SourceData, UInterchangeBaseNodeContainer& BaseNodeContainer) const override;
+
+	/* IInterchangeTexturePayloadInterface Begin */
 
 	/**
 	 * Once the translation is done, the import process need a way to retrieve payload data.
@@ -39,7 +42,9 @@ public:
 	 * @param PayloadKey - The key to retrieve the a particular payload contain into the specified source data.
 	 * @return a PayloadData containing the import image data. The TOptional will not be set if there is an error.
 	 */
-	virtual const TOptional<Interchange::FImportImage> GetPayloadData(const UInterchangeSourceData* SourceData, const FString& PayLoadKey) const override;
+	virtual TOptional<Interchange::FImportImage> GetTexturePayloadData(const UInterchangeSourceData* SourceData, const FString& PayLoadKey) const override;
+
+	/* IInterchangeTexturePayloadInterface End */
 };
 
 
diff --git a/Engine/Plugins/Runtime/Interchange/Source/Import/Public/Texture/InterchangePCXTranslator.h b/Engine/Plugins/Runtime/Interchange/Source/Import/Public/Texture/InterchangePCXTranslator.h
--- a/Engine/Plugins/Runtime/Interchange/Source/Import/Public/Texture/InterchangePCXTranslator.h
+++ b/Engine/Plugins/Runtime/Interchange/Source/Import/Public/Texture/InterchangePCXTranslator.h
@@ -6,13 +6,14 @@
 #include "UObject/Object.h"
 #include "UObject/ObjectMacros.h"
 #include "InterchangeTranslatorBase.h"
+#include "Nodes/InterchangeBaseNodeContainer.h"
 #include "Texture/InterchangeTexturePayloadData.h"
-#include "Texture/InterchangeTextureTranslator.h"
+#include "Texture/InterchangeTexturePayloadInterface.h"
 
 #include "InterchangePCXTranslator.generated.h"
 
 UCLASS(BlueprintType)
-class INTERCHANGEIMPORTPLUGIN_API UInterchangePCXTranslator : public UInterchangeTextureTranslator
+class INTERCHANGEIMPORTPLUGIN_API UInterchangePCXTranslator : public UInterchangeTranslatorBase, public IInterchangeTexturePayloadInterface
 {
 	GENERATED_BODY()
 public:
@@ -29,7 +30,9 @@ public:
 	 * @param BaseNodeContainer - The unreal objects descriptions container where to put the translated source data.
 	 * @return true if the translator can translate the source data, false otherwise.
 	 */
-	virtual bool Translate(const UInterchangeSourceData* SourceData, Interchange::FBaseNodeContainer& BaseNodeContainer) const override;
+	virtual bool Translate(const UInterchangeSourceData* SourceData, UInterchangeBaseNodeContainer& BaseNodeContainer) const override;
+
+	/* IInterchangeTexturePayloadInterface Begin */
 
 	/**
 	 * Once the translation is done, the import process need a way to retrieve payload data.
@@ -39,7 +42,9 @@ public:
 	 * @param PayloadKey - The key to retrieve the a particular payload contain into the specified source data.
 	 * @return a PayloadData containing the import image data. The TOptional will not be set if there is an error.
 	 */
-	virtual const TOptional<Interchange::FImportImage> GetPayloadData(const UInterchangeSourceData* SourceData, const FString& PayLoadKey) const override;
+	virtual TOptional<Interchange::FImportImage> GetTexturePayloadData(const UInterchangeSourceData* SourceData, const FString& PayLoadKey) const override;
+
+	/* IInterchangeTexturePayloadInterface End */
 };
 
 
diff --git a/Engine/Plugins/Runtime/Interchange/Source/Import/Public/Texture/InterchangePNGTranslator.h b/Engine/Plugins/Runtime/Interchange/Source/Import/Public/Texture/InterchangePNGTranslator.h
--- a/Engine/Plugins/Runtime/Interchange/Source/Import/Public/Texture/InterchangePNGTranslator.h
+++ b/Engine/Plugins/Runtime/Interchange/Source/Import/Public/Texture/InterchangePNGTranslator.h
@@ -6,13 +6,14 @@
 #include "UObject/Object.h"
 #include "UObject/ObjectMacros.h"
 #include "InterchangeTranslatorBase.h"
+#include "Nodes/InterchangeBaseNodeContainer.h"
 #include "Texture/InterchangeTexturePayloadData.h"
-#include "Texture/InterchangeTextureTranslator.h"
+#include "Texture/InterchangeTexturePayloadInterface.h"
 
 #include "InterchangePNGTranslator.generated.h"
 
 UCLASS(BlueprintType)
-class INTERCHANGEIMPORTPLUGIN_API UInterchangePNGTranslator : public UInterchangeTextureTranslator
+class INTERCHANGEIMPORTPLUGIN_API UInterchangePNGTranslator : public UInterchangeTranslatorBase, public IInterchangeTexturePayloadInterface
 {
 	GENERATED_BODY()
 public:
@@ -29,7 +30,9 @@ public:
 	 * @param BaseNodeContainer - The unreal objects descriptions container where to put the translated source data.
 	 * @return true if the translator can translate the source data, false otherwise.
 	 */
-	virtual bool Translate(const UInterchangeSourceData* SourceData, Interchange::FBaseNodeContainer& BaseNodeContainer) const override;
+	virtual bool Translate(const UInterchangeSourceData* SourceData, UInterchangeBaseNodeContainer& BaseNodeContainer) const override;
+
+	/* IInterchangeTexturePayloadInterface Begin */
 
 	/**
 	 * Once the translation is done, the import process need a way to retrieve payload data.
@@ -39,7 +42,9 @@ public:
 	 * @param PayloadKey - The key to retrieve the a particular payload contain into the specified source data.
 	 * @return a PayloadData containing the import image data. The TOptional will not be set if there is an error.
 	 */
-	virtual const TOptional<Interchange::FImportImage> GetPayloadData(const UInterchangeSourceData* SourceData, const FString& PayLoadKey) const override;
+	virtual TOptional<Interchange::FImportImage> GetTexturePayloadData(const UInterchangeSourceData* SourceData, const FString& PayLoadKey) const override;
+
+	/* IInterchangeTexturePayloadInterface End */
 };
 
 
diff --git a/Engine/Plugins/Runtime/Interchange/Source/Import/Public/Texture/InterchangePSDTranslator.h b/Engine/Plugins/Runtime/Interchange/Source/Import/Public/Texture/InterchangePSDTranslator.h
--- a/Engine/Plugins/Runtime/Interchange/Source/Import/Public/Texture/InterchangePSDTranslator.h
+++ b/Engine/Plugins/Runtime/Interchange/Source/Import/Public/Texture/InterchangePSDTranslator.h
@@ -6,13 +6,14 @@
 #include "UObject/Object.h"
 #include "UObject/ObjectMacros.h"
 #include "InterchangeTranslatorBase.h"
+#include "Nodes/InterchangeBaseNodeContainer.h"
 #include "Texture/InterchangeTexturePayloadData.h"
-#include "Texture/InterchangeTextureTranslator.h"
+#include "Texture/InterchangeTexturePayloadInterface.h"
 
 #include "InterchangePSDTranslator.generated.h"
 
 UCLASS(BlueprintType)
-class INTERCHANGEIMPORTPLUGIN_API UInterchangePSDTranslator : public UInterchangeTextureTranslator
+class INTERCHANGEIMPORTPLUGIN_API UInterchangePSDTranslator : public UInterchangeTranslatorBase, public IInterchangeTexturePayloadInterface
 {
 	GENERATED_BODY()
 public:
@@ -29,7 +30,9 @@ public:
 	 * @param BaseNodeContainer - The unreal objects descriptions container where to put the translated source data.
 	 * @return true if the translator can translate the source data, false otherwise.
 	 */
-	virtual bool Translate(const UInterchangeSourceData* SourceData, Interchange::FBaseNodeContainer& BaseNodeContainer) const override;
+	virtual bool Translate(const UInterchangeSourceData* SourceData, UInterchangeBaseNodeContainer& BaseNodeContainer) const override;
+
+	/* IInterchangeTexturePayloadInterface Begin */
 
 	/**
 	 * Once the translation is done, the import process need a way to retrieve payload data.
@@ -39,7 +42,9 @@ public:
 	 * @param PayloadKey - The key to retrieve the a particular payload contain into the specified source data.
 	 * @return a PayloadData containing the import image data. The TOptional will not be set if there is an error.
 	 */
-	virtual const TOptional<Interchange::FImportImage> GetPayloadData(const UInterchangeSourceData* SourceData, const FString& PayLoadKey) const override;
+	virtual TOptional<Interchange::FImportImage> GetTexturePayloadData(const UInterchangeSourceData* SourceData, const FString& PayLoadKey) const override;
+
+	/* IInterchangeTexturePayloadInterface End */
 };
 
 
diff --git a/Engine/Plugins/Runtime/Interchange/Source/Import/Public/Texture/InterchangeTGATranslator.h b/Engine/Plugins/Runtime/Interchange/Source/Import/Public/Texture/InterchangeTGATranslator.h
--- a/Engine/Plugins/Runtime/Interchange/Source/Import/Public/Texture/InterchangeTGATranslator.h
+++ b/Engine/Plugins/Runtime/Interchange/Source/Import/Public/Texture/InterchangeTGATranslator.h
@@ -1,18 +1,18 @@
 // Copyright Epic Games, Inc. All Rights Reserved.
-
 #pragma once
 
 #include "CoreMinimal.h"
 #include "UObject/Object.h"
 #include "UObject/ObjectMacros.h"
 #include "InterchangeTranslatorBase.h"
+#include "Nodes/InterchangeBaseNodeContainer.h"
 #include "Texture/InterchangeTexturePayloadData.h"
-#include "Texture/InterchangeTextureTranslator.h"
+#include "Texture/InterchangeTexturePayloadInterface.h"
 
 #include "InterchangeTGATranslator.generated.h"
 
 UCLASS(BlueprintType)
-class INTERCHANGEIMPORTPLUGIN_API UInterchangeTGATranslator : public UInterchangeTextureTranslator
+class INTERCHANGEIMPORTPLUGIN_API UInterchangeTGATranslator : public UInterchangeTranslatorBase, public IInterchangeTexturePayloadInterface
 {
 	GENERATED_BODY()
 public:
@@ -29,7 +29,9 @@ public:
 	 * @param BaseNodeContainer - The unreal objects descriptions container where to put the translated source data.
 	 * @return true if the translator can translate the source data, false otherwise.
 	 */
-	virtual bool Translate(const UInterchangeSourceData* SourceData, Interchange::FBaseNodeContainer& BaseNodeContainer) const override;
+	virtual bool Translate(const UInterchangeSourceData* SourceData, UInterchangeBaseNodeContainer& BaseNodeContainer) const override;
+
+	/* IInterchangeTexturePayloadInterface Begin */
 
 	/**
 	 * Once the translation is done, the import process need a way to retrieve payload data.
@@ -39,7 +41,9 @@ public:
 	 * @param PayloadKey - The key to retrieve the a particular payload contain into the specified source data.
 	 * @return a PayloadData containing the import image data. The TOptional will not be set if there is an error.
 	 */
-	virtual const TOptional<Interchange::FImportImage> GetPayloadData(const UInterchangeSourceData* SourceData, const FString& PayLoadKey) const override;
+	virtual TOptional<Interchange::FImportImage> GetTexturePayloadData(const UInterchangeSourceData* SourceData, const FString& PayLoadKey) const override;
+
+	/* IInterchangeTexturePayloadInterface End */
 };
 
 
diff --git a/Engine/Plugins/Runtime/Interchange/Source/Import/Public/Texture/InterchangeTextureFactory.h b/Engine/Plugins/Runtime/Interchange/Source/Import/Public/Texture/InterchangeTextureFactory.h
--- a/Engine/Plugins/Runtime/Interchange/Source/Import/Public/Texture/InterchangeTextureFactory.h
+++ b/Engine/Plugins/Runtime/Interchange/Source/Import/Public/Texture/InterchangeTextureFactory.h
@@ -3,7 +3,6 @@
 #pragma once
 
 #include "CoreMinimal.h"
-#include "Engine/Texture.h"
 #include "InterchangeFactoryBase.h"
 #include "UObject/Object.h"
 #include "UObject/ObjectMacros.h"
diff --git a/Engine/Plugins/Runtime/Interchange/Source/Import/Public/Texture/InterchangeTexturePayloadInterface.h b/Engine/Plugins/Runtime/Interchange/Source/Import/Public/Texture/InterchangeTexturePayloadInterface.h
new file mode 100644
--- /dev/null
+++ b/Engine/Plugins/Runtime/Interchange/Source/Import/Public/Texture/InterchangeTexturePayloadInterface.h
@@ -0,0 +1,36 @@
+// Copyright Epic Games, Inc. All Rights Reserved.
+#pragma once
+
+#include "CoreMinimal.h"
+#include "UObject/Object.h"
+#include "UObject/ObjectMacros.h"
+#include "Texture/InterchangeTexturePayloadData.h"
+
+#include "InterchangeTexturePayloadInterface.generated.h"
+
+UINTERFACE()
+class INTERCHANGEIMPORTPLUGIN_API UInterchangeTexturePayloadInterface : public UInterface
+{
+	GENERATED_BODY()
+};
+
+/**
+ * Texture payload interface. Derive from it if your payload can import texture
+ */
+class INTERCHANGEIMPORTPLUGIN_API IInterchangeTexturePayloadInterface
+{
+	GENERATED_BODY()
+public:
+
+	/**
+	 * Once the translation is done, the import process need a way to retrieve payload data.
+	 * This payload will be use by the factories to create the asset.
+	 *
+	 * @param SourceData - The source data containing the data to translate
+	 * @param PayloadKey - The key to retrieve the a particular payload contain into the specified source data.
+	 * @return a PayloadData containing the import image data. The TOptional will not be set if there is an error.
+	 */
+	virtual TOptional<Interchange::FImportImage> GetTexturePayloadData(const UInterchangeSourceData* SourceData, const FString& PayloadKey) const = 0;
+};
+
+
diff --git a/Engine/Plugins/Runtime/Interchange/Source/Import/Public/Texture/InterchangeTextureTranslator.h b/Engine/Plugins/Runtime/Interchange/Source/Import/Public/Texture/InterchangeTextureTranslator.h
deleted file mode 100644
diff --git a/Engine/Plugins/Runtime/Interchange/Source/Nodes/Public/InterchangeMaterialNode.h b/Engine/Plugins/Runtime/Interchange/Source/Nodes/Public/InterchangeMaterialNode.h
new file mode 100644
--- /dev/null
+++ b/Engine/Plugins/Runtime/Interchange/Source/Nodes/Public/InterchangeMaterialNode.h
@@ -0,0 +1,400 @@
+// Copyright Epic Games, Inc. All Rights Reserved.
+
+#pragma once
+
+#include "CoreMinimal.h"
+#include "Nodes/InterchangeBaseNode.h"
+
+#if WITH_ENGINE
+
+#include "Materials/Material.h"
+#include "Materials/MaterialInstance.h"
+#include "Materials/MaterialInterface.h"
+
+#endif
+
+#include "InterchangeMaterialNode.generated.h"
+
+//Interchange namespace
+namespace Interchange
+{
+	
+	struct FMaterialNodeStaticData : public FBaseNodeStaticData
+	{
+		static const FAttributeKey& PayloadSourceFileKey()
+		{
+			static FAttributeKey AttributeKey(TEXT("__PayloadSourceFile__"));
+			return AttributeKey;
+		}
+	};
+} //ns Interchange
+
+/**
+ * enum use to declare supported parameter
+ */
+UENUM(BlueprintType)
+enum class EInterchangeMaterialNodeParameterName : uint8
+{
+	BaseColor,
+	EmissiveColor,
+	Specular,
+	Roughness,
+	Metallic,
+	Normal,
+	Opacity,
+	OpacityMask
+};
+
+UCLASS(BlueprintType)
+class INTERCHANGENODEPLUGIN_API UInterchangeMaterialNode : public UInterchangeBaseNode
+{
+	GENERATED_BODY()
+
+public:
+	UInterchangeMaterialNode()
+	:UInterchangeBaseNode()
+	{
+#if WITH_ENGINE
+		AssetClass = nullptr;
+#endif
+	}
+
+	/**
+	 * Initialize node data
+	 * @param: UniqueID - The uniqueId for this node
+	 * @param DisplayLabel - The name of the node
+	 * @param InAssetClass - The class the material factory will create for this node.
+	 */
+	UFUNCTION(BlueprintCallable, Category = "Interchange | Node | Material")
+	void InitializeMaterialNode(const FName& UniqueID, const FName& DisplayLabel, const FString& InAssetClass)
+	{
+		bIsMaterialNodeClassInitialized = false;
+		InitializeNode(UniqueID, DisplayLabel);
+		FString OperationName = GetTypeName() + TEXT(".SetAssetClassName");
+		InterchangePrivateNodeBase::SetCustomAttribute<FString>(Attributes, ClassNameAttributeKey, OperationName, InAssetClass);
+		FillAssetClassFromAttribute();
+	}
+
+	/**
+	 * Return the node type name of the class, we use this when reporting error
+	 */
+	virtual FString GetTypeName() const override
+	{
+		const FString TypeName = TEXT("MaterialNode");
+		return TypeName;
+	}
+
+	/** Get the class this node want to create */
+	UFUNCTION(BlueprintCallable, Category = "Interchange | Node | Material")
+	virtual class UClass* GetAssetClass() const override
+	{
+		ensure(bIsMaterialNodeClassInitialized);
+#if WITH_ENGINE
+		return AssetClass.Get() != nullptr ? AssetClass.Get() : UMaterial::StaticClass();
+#else
+		return nullptr;
+#endif
+	}
+
+	virtual FGuid GetHash() const override
+	{
+		return Attributes.GetStorageHash();
+	}
+
+	virtual const TOptional<FString> GetPayLoadKey() const
+	{
+		if (!Attributes.ContainAttribute(Interchange::FMaterialNodeStaticData::PayloadSourceFileKey()))
+		{
+			return TOptional<FString>();
+		}
+		Interchange::FAttributeStorage::TAttributeHandle<FString> AttributeHandle = Attributes.GetAttributeHandle<FString>(Interchange::FMaterialNodeStaticData::PayloadSourceFileKey());
+		if (!AttributeHandle.IsValid())
+		{
+			return TOptional<FString>();
+		}
+		FString PayloadKey;
+		Interchange::EAttributeStorageResult Result = AttributeHandle.Get(PayloadKey);
+		if (!IsAttributeStorageResultSuccess(Result))
+		{
+			LogAttributeStorageErrors(Result, TEXT("UInterchangeMaterialNode.GetPayLoadKey"), Interchange::FMaterialNodeStaticData::PayloadSourceFileKey());
+			return TOptional<FString>();
+		}
+		return TOptional<FString>(PayloadKey);
+	}
+
+	UFUNCTION(BlueprintCallable, Category = "Interchange | Node | Material")
+	virtual void SetPayLoadKey(const FString& PayloadKey)
+	{
+		Interchange::EAttributeStorageResult Result = Attributes.RegisterAttribute(Interchange::FMaterialNodeStaticData::PayloadSourceFileKey(), PayloadKey);
+		if (!IsAttributeStorageResultSuccess(Result))
+		{
+			LogAttributeStorageErrors(Result, TEXT("UInterchangeMaterialNode.SetPayLoadKey"), Interchange::FMaterialNodeStaticData::PayloadSourceFileKey());
+		}
+	}
+
+	/** Return false if the Attribute was not set previously.*/
+	UFUNCTION(BlueprintCallable, Category = "Interchange | Node | Material")
+	bool GetCustomMaterialUsage(uint8& AttributeValue) const
+	{
+		IMPLEMENT_NODE_ATTRIBUTE_GETTER(MaterialUsage, uint8);
+	}
+
+	UFUNCTION(BlueprintCallable, Category = "Interchange | Node | Material")
+	bool SetCustomMaterialUsage(const uint8& AttributeValue, bool bAddApplyDelegate = true)
+	{
+		IMPLEMENT_NODE_ATTRIBUTE_SETTER(UInterchangeMaterialNode, MaterialUsage, uint8, UMaterialInterface);
+	}
+
+	/** Return false if the Attribute was not set previously.*/
+	UFUNCTION(BlueprintCallable, Category = "Interchange | Node | Material")
+	bool GetCustomBlendMode(uint8& AttributeValue) const
+	{
+		IMPLEMENT_NODE_ATTRIBUTE_GETTER(BlendMode, uint8);
+	}
+
+	UFUNCTION(BlueprintCallable, Category = "Interchange | Node | Material")
+	bool SetCustomBlendMode(const uint8& AttributeValue, bool bAddApplyDelegate = true)
+	{
+		IMPLEMENT_NODE_ATTRIBUTE_SETTER(UInterchangeMaterialNode, BlendMode, uint8, UMaterialInterface);
+	}
+
+	//////////////////////////////////////////////////////////////////////////
+	/**
+	 * Parameter interface for c++ python and blueprint
+	 * Each added parameter will create a material expression input later in the material factory.
+	 * You can control the imput type to use: texture sampler, vector parameter or scalar parameter
+	 */
+
+
+	/**
+	 * Add a texture parameter for the specified ParameterName.
+	 * @param ParameterName - The parameter we want to set the texture for.
+	 * @param TextureUID - The texture node uniqueID that has the texture we want to set to the specified parameter
+	 * @param UVSetNameUseByTextureInput - [Optional] The fbx name of the UV set we want the material input expression use to sample the specified texture.
+	 *                                     It can be empty, the material factory should in that case map it on the UV channel 0
+	 * @note - A parameter name can have only one of the 3 type of input set, the last input type set is the one that will be created.
+	 */
+	UFUNCTION(BlueprintCallable, Category = "Interchange | Node | Material")
+	void AddTextureParameterData(const EInterchangeMaterialNodeParameterName ParameterName, const FName TextureUID, const int32 UVSetIndex, const float ScaleU, const float ScaleV)
+	{
+		FParameterData& ParameterData = ParameterDatas.FindOrAdd(ParameterName);
+		ParameterData.bIsTextureParameter = true;
+		ParameterData.bIsVectorParameter = false;
+		ParameterData.bIsScalarParameter = false;
+		ParameterData.TextureUID = TextureUID;
+		ParameterData.UVSetIndex = UVSetIndex;
+		ParameterData.ScaleU = ScaleU;
+		ParameterData.ScaleV = ScaleV;
+	}
+
+	UFUNCTION(BlueprintCallable, Category = "Interchange | Node | Material")
+	bool GetTextureParameterData(const EInterchangeMaterialNodeParameterName ParameterName, FName& OutTextureUID, int32& OutUVSetIndex, float& OutScaleU, float& OutScaleV) const
+	{
+		if (const FParameterData* ParameterData = ParameterDatas.Find(ParameterName))
+		{
+			if (!ParameterData->bIsTextureParameter)
+			{
+				return false;
+			}
+			OutTextureUID = ParameterData->TextureUID;
+			OutUVSetIndex = ParameterData->UVSetIndex;
+			OutScaleU = ParameterData->ScaleU;
+			OutScaleV = ParameterData->ScaleV;
+			return true;
+		}
+
+		return false;
+	}
+
+	UFUNCTION(BlueprintCallable, Category = "Interchange | Node | Material")
+	void AddVectorParameterData(const EInterchangeMaterialNodeParameterName ParameterName, const FVector& VectorData)
+	{
+		FParameterData& ParameterData = ParameterDatas.FindOrAdd(ParameterName);
+		ParameterData.bIsTextureParameter = false;
+		ParameterData.bIsVectorParameter = true;
+		ParameterData.bIsScalarParameter = false;
+		ParameterData.VectorParameter = VectorData;
+	}
+
+	UFUNCTION(BlueprintCallable, Category = "Interchange | Node | Material")
+	bool GetVectorParameterData(const EInterchangeMaterialNodeParameterName ParameterName, FVector& OutVectorData) const
+	{
+		if (const FParameterData* ParameterData = ParameterDatas.Find(ParameterName))
+		{
+			if (!ParameterData->bIsVectorParameter)
+			{
+				return false;
+			}
+			OutVectorData = ParameterData->VectorParameter;
+			return true;
+		}
+		return false;
+	}
+
+	UFUNCTION(BlueprintCallable, Category = "Interchange | Node | Material")
+	void AddScalarParameterData(const EInterchangeMaterialNodeParameterName ParameterName, const float ScalarData)
+	{
+		FParameterData& ParameterData = ParameterDatas.FindOrAdd(ParameterName);
+		ParameterData.bIsTextureParameter = false;
+		ParameterData.bIsVectorParameter = false;
+		ParameterData.bIsScalarParameter = true;
+		ParameterData.ScalarParameter = ScalarData;
+	}
+
+	UFUNCTION(BlueprintCallable, Category = "Interchange | Node | Material")
+	bool GetScalarParameterData(const EInterchangeMaterialNodeParameterName ParameterName, float& OutScalarData) const
+	{
+		if (const FParameterData* ParameterData = ParameterDatas.Find(ParameterName))
+		{
+			if (!ParameterData->bIsScalarParameter)
+			{
+				return false;
+			}
+			OutScalarData = ParameterData->ScalarParameter;
+			return true;
+		}
+		return false;
+	}
+
+	virtual void Serialize(FArchive& Ar) override
+	{
+		Super::Serialize(Ar);
+
+		if (Ar.IsLoading())
+		{
+			ParameterDatas.Reset();
+		}
+		int32 ParamCount = ParameterDatas.Num();
+		Ar << ParamCount;
+
+		if (Ar.IsSaving())
+		{
+			for (TPair<EInterchangeMaterialNodeParameterName, FParameterData> ParamPair : ParameterDatas)
+			{
+				Ar << ParamPair.Key;
+				Ar << ParamPair.Value;
+			}
+		}
+		else if (Ar.IsLoading())
+		{
+			for (int32 ParamIndex = 0; ParamIndex < ParamCount; ++ParamIndex)
+			{
+				EInterchangeMaterialNodeParameterName ParamKey;
+				Ar << ParamKey;
+				FParameterData& ParamValue = ParameterDatas.FindOrAdd(ParamKey);
+				Ar << ParamValue;
+			}
+		}
+
+#if WITH_ENGINE
+		if (Ar.IsLoading())
+		{
+			//Make sure the class is properly set when we compile with engine, this will set the
+			//bIsMaterialNodeClassInitialized to true.
+			SetMaterialNodeClassFromClassAttribute();
+		}
+#endif //#if WITH_ENGINE
+	}
+
+private:
+
+	void FillAssetClassFromAttribute()
+	{
+#if WITH_ENGINE
+		FString OperationName = GetTypeName() + TEXT(".GetAssetClassName");
+		FString ClassName;
+		InterchangePrivateNodeBase::GetCustomAttribute<FString>(Attributes, ClassNameAttributeKey, OperationName, ClassName);
+		if (ClassName.Equals(UMaterial::StaticClass()->GetName()))
+		{
+			AssetClass = UMaterial::StaticClass();
+			bIsMaterialNodeClassInitialized = true;
+		}
+		else if (ClassName.Equals(UMaterialInstance::StaticClass()->GetName()))
+		{
+			AssetClass = UMaterialInstance::StaticClass();
+			bIsMaterialNodeClassInitialized = true;
+		}
+#endif
+	}
+
+	bool SetMaterialNodeClassFromClassAttribute()
+	{
+		if (!bIsMaterialNodeClassInitialized)
+		{
+			FillAssetClassFromAttribute();
+		}
+		return bIsMaterialNodeClassInitialized;
+	}
+
+	struct FParameterData
+	{
+		bool bIsTextureParameter = false;
+		FName TextureUID = NAME_None;
+		int32 UVSetIndex = 0; //The "UV set" index we use to set the material sampler input default is channel 0
+		float ScaleU = 1.0f;
+		float ScaleV = 1.0f;
+
+		bool bIsVectorParameter = false;
+		FVector VectorParameter = FVector(0.0f);
+
+		bool bIsScalarParameter = false;
+		float ScalarParameter = 0.0f;
+
+		friend FArchive& operator<<(FArchive& Ar, FParameterData& ParameterData)
+		{
+			Ar << ParameterData.bIsTextureParameter;
+			Ar << ParameterData.TextureUID;
+			Ar << ParameterData.UVSetIndex;
+			Ar << ParameterData.ScaleU;
+			Ar << ParameterData.ScaleV;
+
+			Ar << ParameterData.bIsVectorParameter;
+			Ar << ParameterData.VectorParameter;
+
+			Ar << ParameterData.bIsScalarParameter;
+			Ar << ParameterData.ScalarParameter;
+			return Ar;
+		}
+	};
+
+	//This member is serialize manually in the serialize override
+	TMap<EInterchangeMaterialNodeParameterName, FParameterData> ParameterDatas;
+
+	const Interchange::FAttributeKey ClassNameAttributeKey = Interchange::FAttributeKey(TEXT("__ClassTypeAttribute__"));
+
+	//Material Attribute
+	const Interchange::FAttributeKey Macro_CustomMaterialUsageKey = Interchange::FAttributeKey(TEXT("MaterialUsage"));
+	const Interchange::FAttributeKey Macro_CustomBlendModeKey = Interchange::FAttributeKey(TEXT("BlendMode"));
+
+#if WITH_ENGINE
+	bool ApplyCustomMaterialUsageToAsset(UObject * Asset) const
+	{
+		if (!Asset)
+		{
+			return false;
+		}
+		UMaterialInterface* TypedObject = Cast<UMaterialInterface>(Asset);
+		if (!TypedObject)
+		{
+			return false;
+		}
+		uint8 ValueData;
+		if (GetCustomMaterialUsage(ValueData))
+		{
+			TypedObject->CheckMaterialUsage_Concurrent(EMaterialUsage(ValueData));
+			return true;
+		}
+		return false;
+	}
+
+#endif
+
+	IMPLEMENT_NODE_ATTRIBUTE_APPLY_UOBJECT(BlendMode, uint8, UMaterial, TEnumAsByte<enum EBlendMode>);
+
+protected:
+#if WITH_ENGINE
+	TSubclassOf<UTexture> AssetClass = nullptr;
+#endif
+	bool bIsMaterialNodeClassInitialized = false;
+
+};
diff --git a/Engine/Plugins/Runtime/Interchange/Source/Nodes/Public/InterchangeSceneNode.h b/Engine/Plugins/Runtime/Interchange/Source/Nodes/Public/InterchangeSceneNode.h
new file mode 100644
--- /dev/null
+++ b/Engine/Plugins/Runtime/Interchange/Source/Nodes/Public/InterchangeSceneNode.h
@@ -0,0 +1,71 @@
+// Copyright Epic Games, Inc. All Rights Reserved.
+
+#pragma once
+
+#include "CoreMinimal.h"
+#include "Nodes/InterchangeBaseNode.h"
+
+#if WITH_ENGINE
+
+
+#endif
+
+#include "InterchangeSceneNode.generated.h"
+
+
+UCLASS(BlueprintType)
+class INTERCHANGENODEPLUGIN_API UInterchangeSceneNode : public UInterchangeBaseNode
+{
+	GENERATED_BODY()
+
+public:
+	UInterchangeSceneNode()
+	:UInterchangeBaseNode()
+	{
+	}
+
+	/**
+	 * Return the node type name of the class, we use this when reporting error
+	 */
+	virtual FString GetTypeName() const override
+	{
+		const FString TypeName = TEXT("SceneNode");
+		return TypeName;
+	}
+
+	virtual FGuid GetHash() const override
+	{
+		return Attributes.GetStorageHash();
+	}
+
+	/** Return false if the Attribute was not set previously.*/
+	UFUNCTION(BlueprintCallable, Category = "Interchange | Node | Scene")
+	bool GetCustomLocalTransform(FTransform& AttributeValue) const
+	{
+		IMPLEMENT_NODE_ATTRIBUTE_GETTER(LocalTransform, FTransform);
+	}
+
+	UFUNCTION(BlueprintCallable, Category = "Interchange | Node | Scene")
+	bool SetCustomLocalTransform(const FTransform& AttributeValue)
+	{
+		IMPLEMENT_NODE_ATTRIBUTE_SETTER_NODELEGATE(LocalTransform, FTransform);
+	}
+
+	/** Return false if the Attribute was not set previously.*/
+	UFUNCTION(BlueprintCallable, Category = "Interchange | Node | Scene")
+	bool GetCustomGlobalTransform(FTransform& AttributeValue) const
+	{
+		IMPLEMENT_NODE_ATTRIBUTE_GETTER(GlobalTransform, FTransform);
+	}
+
+	UFUNCTION(BlueprintCallable, Category = "Interchange | Node | Scene")
+	bool SetCustomGlobalTransform(const FTransform& AttributeValue)
+	{
+		IMPLEMENT_NODE_ATTRIBUTE_SETTER_NODELEGATE(GlobalTransform, FTransform);
+	}
+
+private:
+	//Scene Attribute Keys
+	const Interchange::FAttributeKey Macro_CustomLocalTransformKey = Interchange::FAttributeKey(TEXT("LocalTransform"));
+	const Interchange::FAttributeKey Macro_CustomGlobalTransformKey = Interchange::FAttributeKey(TEXT("GlobalTransform"));
+};
diff --git a/Engine/Plugins/Runtime/Interchange/Source/Nodes/Public/InterchangeTextureNode.h b/Engine/Plugins/Runtime/Interchange/Source/Nodes/Public/InterchangeTextureNode.h
new file mode 100644
--- /dev/null
+++ b/Engine/Plugins/Runtime/Interchange/Source/Nodes/Public/InterchangeTextureNode.h
@@ -0,0 +1,897 @@
+// Copyright Epic Games, Inc. All Rights Reserved.
+
+#pragma once
+
+#include "CoreMinimal.h"
+#include "Nodes/InterchangeBaseNode.h"
+
+#if WITH_ENGINE
+#include "Engine/Texture.h"
+#include "Engine/Texture2D.h"
+#endif
+
+#include "InterchangeTextureNode.generated.h"
+
+//Interchange namespace
+namespace Interchange
+{
+
+struct FTextureNodeStaticData : public FBaseNodeStaticData
+{
+	static const FAttributeKey& PayloadSourceFileKey()
+	{
+		static FAttributeKey AttributeKey(TEXT("__PayloadSourceFile__"));
+		return AttributeKey;
+	}
+};
+
+}
+
+UCLASS(BlueprintType)
+class INTERCHANGENODEPLUGIN_API UInterchangeTextureNode : public UInterchangeBaseNode
+{
+	GENERATED_BODY()
+
+public:
+	UInterchangeTextureNode()
+	:UInterchangeBaseNode()
+	{
+#if WITH_ENGINE
+		AssetClass = nullptr;
+#endif
+	}
+
+	/**
+	 * Initialize node data
+	 * @param: UniqueID - The uniqueId for this node
+	 * @param DisplayLabel - The name of the node
+	 * @param InAssetClass - The class the texture factory will create for this node.
+	 */
+	UFUNCTION(BlueprintCallable, Category = "Interchange | Node | Texture")
+	void InitializeTextureNode(const FName& UniqueID, const FName& DisplayLabel, const FString& InAssetClass)
+	{
+		bIsTextureNodeClassInitialized = false;
+		InitializeNode(UniqueID, DisplayLabel);
+		FString OperationName = GetTypeName() + TEXT(".SetAssetClassName");
+		InterchangePrivateNodeBase::SetCustomAttribute<FString>(Attributes, ClassNameAttributeKey, OperationName, InAssetClass);
+		FillAssetClassFromAttribute();
+	}
+
+	virtual void Serialize(FArchive& Ar) override
+	{
+		Super::Serialize(Ar);
+#if WITH_ENGINE
+		if(Ar.IsLoading())
+		{
+			//Make sure the class is properly set when we compile with engine, this will set the
+			//bIsTextureNodeClassInitialized to true.
+			SetTextureNodeClassFromClassAttribute();
+		}
+#endif //#if WITH_ENGINE
+	}
+
+	/**
+	 * Return the node type name of the class, we use this when reporting error
+	 */
+	virtual FString GetTypeName() const override
+	{
+		const FString TypeName = TEXT("TextureNode");
+		return TypeName;
+	}
+
+	/** Get the class this node want to create */
+	UFUNCTION(BlueprintCallable, Category = "Interchange | Node | Texture")
+	virtual class UClass* GetAssetClass() const override
+	{
+		ensure(bIsTextureNodeClassInitialized);
+#if WITH_ENGINE
+		return AssetClass.Get() != nullptr ? AssetClass.Get() : UTexture::StaticClass();
+#else
+		return nullptr;
+#endif
+	}
+
+	virtual FGuid GetHash() const override
+	{
+		return Attributes.GetStorageHash();
+	}
+
+	/** Texture node Interface Begin */
+	virtual const TOptional<FString> GetPayLoadKey() const
+	{
+		if (!Attributes.ContainAttribute(Interchange::FTextureNodeStaticData::PayloadSourceFileKey()))
+		{
+			return TOptional<FString>();
+		}
+		Interchange::FAttributeStorage::TAttributeHandle<FString> AttributeHandle = Attributes.GetAttributeHandle<FString>(Interchange::FTextureNodeStaticData::PayloadSourceFileKey());
+		if (!AttributeHandle.IsValid())
+		{
+			return TOptional<FString>();
+		}
+		FString PayloadKey;
+		Interchange::EAttributeStorageResult Result = AttributeHandle.Get(PayloadKey);
+		if (!IsAttributeStorageResultSuccess(Result))
+		{
+			LogAttributeStorageErrors(Result, TEXT("UInterchangeTextureNode.GetPayLoadKey"), Interchange::FTextureNodeStaticData::PayloadSourceFileKey());
+			return TOptional<FString>();
+		}
+		return TOptional<FString>(PayloadKey);
+	}
+
+	UFUNCTION(BlueprintCallable, Category = "Interchange | Node | Texture")
+	virtual void SetPayLoadKey(const FString& PayloadKey)
+	{
+		Interchange::EAttributeStorageResult Result = Attributes.RegisterAttribute(Interchange::FTextureNodeStaticData::PayloadSourceFileKey(), PayloadKey);
+		if (!IsAttributeStorageResultSuccess(Result))
+		{
+			LogAttributeStorageErrors(Result, TEXT("UInterchangeTextureNode.SetPayLoadKey"), Interchange::FTextureNodeStaticData::PayloadSourceFileKey());
+		}
+	}
+
+
+public:
+	/** Return false if the Attribute was not set previously.*/
+	UFUNCTION(BlueprintCallable, Category = "Interchange | Node | Texture")
+	bool GetCustomAdjustBrightness(float& AttributeValue) const
+	{
+		IMPLEMENT_NODE_ATTRIBUTE_GETTER(AdjustBrightness, float);
+	}
+
+	UFUNCTION(BlueprintCallable, Category = "Interchange | Node | Texture")
+	bool SetCustomAdjustBrightness(const float& AttributeValue, bool bAddApplyDelegate = true)
+	{
+#if WITH_EDITORONLY_DATA
+		IMPLEMENT_NODE_ATTRIBUTE_SETTER(UInterchangeTextureNode, AdjustBrightness, float, UTexture)
+#else
+		IMPLEMENT_NODE_ATTRIBUTE_SETTER_NODELEGATE(AdjustBrightness, float)
+#endif
+	}
+
+	/** Return false if the Attribute was not set previously.*/
+	UFUNCTION(BlueprintCallable, Category = "Interchange | Node | Texture")
+	bool GetCustomAdjustBrightnessCurve(float& AttributeValue) const
+	{
+		IMPLEMENT_NODE_ATTRIBUTE_GETTER(AdjustBrightnessCurve, float);
+	}
+
+	UFUNCTION(BlueprintCallable, Category = "Interchange | Node | Texture")
+	bool SetCustomAdjustBrightnessCurve(const float& AttributeValue, bool bAddApplyDelegate = true)
+	{
+#if WITH_EDITORONLY_DATA
+		IMPLEMENT_NODE_ATTRIBUTE_SETTER(UInterchangeTextureNode, AdjustBrightnessCurve, float, UTexture)
+#else
+		IMPLEMENT_NODE_ATTRIBUTE_SETTER_NODELEGATE(AdjustBrightnessCurve, float)
+#endif
+	}
+
+	/** Return false if the Attribute was not set previously.*/
+	UFUNCTION(BlueprintCallable, Category = "Interchange | Node | Texture")
+	bool GetCustomAdjustVibrance(float& AttributeValue) const
+	{
+		IMPLEMENT_NODE_ATTRIBUTE_GETTER(AdjustVibrance, float);
+	}
+
+	UFUNCTION(BlueprintCallable, Category = "Interchange | Node | Texture")
+	bool SetCustomAdjustVibrance(const float& AttributeValue, bool bAddApplyDelegate = true)
+	{
+#if WITH_EDITORONLY_DATA
+		IMPLEMENT_NODE_ATTRIBUTE_SETTER(UInterchangeTextureNode, AdjustVibrance, float, UTexture)
+#else
+		IMPLEMENT_NODE_ATTRIBUTE_SETTER_NODELEGATE(AdjustVibrance, float)
+#endif
+	}
+
+	/** Return false if the Attribute was not set previously.*/
+	UFUNCTION(BlueprintCallable, Category = "Interchange | Node | Texture")
+	bool GetCustomAdjustSaturation(float& AttributeValue) const
+	{
+		IMPLEMENT_NODE_ATTRIBUTE_GETTER(AdjustSaturation, float);
+	}
+
+	UFUNCTION(BlueprintCallable, Category = "Interchange | Node | Texture")
+	bool SetCustomAdjustSaturation(const float& AttributeValue, bool bAddApplyDelegate = true)
+	{
+#if WITH_EDITORONLY_DATA
+		IMPLEMENT_NODE_ATTRIBUTE_SETTER(UInterchangeTextureNode, AdjustSaturation, float, UTexture)
+#else
+		IMPLEMENT_NODE_ATTRIBUTE_SETTER_NODELEGATE(AdjustSaturation, float)
+#endif
+	}
+
+	/** Return false if the Attribute was not set previously.*/
+	UFUNCTION(BlueprintCallable, Category = "Interchange | Node | Texture")
+	bool GetCustomAdjustRGBCurve(float& AttributeValue) const
+	{
+		IMPLEMENT_NODE_ATTRIBUTE_GETTER(AdjustRGBCurve, float);
+	}
+
+	UFUNCTION(BlueprintCallable, Category = "Interchange | Node | Texture")
+	bool SetCustomAdjustRGBCurve(const float& AttributeValue, bool bAddApplyDelegate = true)
+	{
+#if WITH_EDITORONLY_DATA
+		IMPLEMENT_NODE_ATTRIBUTE_SETTER(UInterchangeTextureNode, AdjustRGBCurve, float, UTexture)
+#else
+		IMPLEMENT_NODE_ATTRIBUTE_SETTER_NODELEGATE(AdjustRGBCurve, float)
+#endif
+	}
+
+	/** Return false if the Attribute was not set previously.*/
+	UFUNCTION(BlueprintCallable, Category = "Interchange | Node | Texture")
+	bool GetCustomAdjustHue(float& AttributeValue) const
+	{
+		IMPLEMENT_NODE_ATTRIBUTE_GETTER(AdjustHue, float);
+	}
+
+	UFUNCTION(BlueprintCallable, Category = "Interchange | Node | Texture")
+	bool SetCustomAdjustHue(const float& AttributeValue, bool bAddApplyDelegate = true)
+	{
+#if WITH_EDITORONLY_DATA
+		IMPLEMENT_NODE_ATTRIBUTE_SETTER(UInterchangeTextureNode, AdjustHue, float, UTexture)
+#else
+		IMPLEMENT_NODE_ATTRIBUTE_SETTER_NODELEGATE(AdjustHue, float)
+#endif
+	}
+
+	/** Return false if the Attribute was not set previously.*/
+	UFUNCTION(BlueprintCallable, Category = "Interchange | Node | Texture")
+	bool GetCustomAdjustMinAlpha(float& AttributeValue) const
+	{
+		IMPLEMENT_NODE_ATTRIBUTE_GETTER(AdjustMinAlpha, float);
+	}
+
+	UFUNCTION(BlueprintCallable, Category = "Interchange | Node | Texture")
+	bool SetCustomAdjustMinAlpha(const float& AttributeValue, bool bAddApplyDelegate = true)
+	{
+#if WITH_EDITORONLY_DATA
+		IMPLEMENT_NODE_ATTRIBUTE_SETTER(UInterchangeTextureNode, AdjustMinAlpha, float, UTexture)
+#else
+		IMPLEMENT_NODE_ATTRIBUTE_SETTER_NODELEGATE(AdjustMinAlpha, float)
+#endif
+	}
+
+	/** Return false if the Attribute was not set previously.*/
+	UFUNCTION(BlueprintCallable, Category = "Interchange | Node | Texture")
+	bool GetCustomAdjustMaxAlpha(float& AttributeValue) const
+	{
+		IMPLEMENT_NODE_ATTRIBUTE_GETTER(AdjustMaxAlpha, float);
+	}
+
+	UFUNCTION(BlueprintCallable, Category = "Interchange | Node | Texture")
+	bool SetCustomAdjustMaxAlpha(const float& AttributeValue, bool bAddApplyDelegate = true)
+	{
+#if WITH_EDITORONLY_DATA
+		IMPLEMENT_NODE_ATTRIBUTE_SETTER(UInterchangeTextureNode, AdjustMaxAlpha, float, UTexture)
+#else
+		IMPLEMENT_NODE_ATTRIBUTE_SETTER_NODELEGATE(AdjustMaxAlpha, float)
+#endif
+	}
+
+	/** Return false if the Attribute was not set previously.*/
+	UFUNCTION(BlueprintCallable, Category = "Interchange | Node | Texture")
+	bool GetCustombChromaKeyTexture(bool& AttributeValue) const
+	{
+		IMPLEMENT_NODE_ATTRIBUTE_GETTER(bChromaKeyTexture, bool);
+	}
+
+	UFUNCTION(BlueprintCallable, Category = "Interchange | Node | Texture")
+	bool SetCustombChromaKeyTexture(const bool& AttributeValue, bool bAddApplyDelegate = true)
+	{
+#if WITH_EDITORONLY_DATA
+		IMPLEMENT_NODE_ATTRIBUTE_SETTER(UInterchangeTextureNode, bChromaKeyTexture, bool, UTexture)
+#else
+		IMPLEMENT_NODE_ATTRIBUTE_SETTER_NODELEGATE(bChromaKeyTexture, bool)
+#endif
+	}
+
+	/** Return false if the Attribute was not set previously.*/
+	UFUNCTION(BlueprintCallable, Category = "Interchange | Node | Texture")
+	bool GetCustomChromaKeyThreshold(float& AttributeValue) const
+	{
+		IMPLEMENT_NODE_ATTRIBUTE_GETTER(ChromaKeyThreshold, float);
+	}
+
+	UFUNCTION(BlueprintCallable, Category = "Interchange | Node | Texture")
+	bool SetCustomChromaKeyThreshold(const float& AttributeValue, bool bAddApplyDelegate = true)
+	{
+#if WITH_EDITORONLY_DATA
+		IMPLEMENT_NODE_ATTRIBUTE_SETTER(UInterchangeTextureNode, ChromaKeyThreshold, float, UTexture)
+#else
+		IMPLEMENT_NODE_ATTRIBUTE_SETTER_NODELEGATE(ChromaKeyThreshold, float)
+#endif
+	}
+
+	/** Return false if the Attribute was not set previously.*/
+	UFUNCTION(BlueprintCallable, Category = "Interchange | Node | Texture")
+	bool GetCustomChromaKeyColor(FColor& AttributeValue) const
+	{
+		IMPLEMENT_NODE_ATTRIBUTE_GETTER(ChromaKeyColor, FColor);
+	}
+
+	UFUNCTION(BlueprintCallable, Category = "Interchange | Node | Texture")
+	bool SetCustomChromaKeyColor(const FColor& AttributeValue, bool bAddApplyDelegate = true)
+	{
+#if WITH_EDITORONLY_DATA
+		IMPLEMENT_NODE_ATTRIBUTE_SETTER(UInterchangeTextureNode, ChromaKeyColor, FColor, UTexture)
+#else
+		IMPLEMENT_NODE_ATTRIBUTE_SETTER_NODELEGATE(ChromaKeyColor, FColor)
+#endif
+	}
+
+	//////////////////////////////////////////////////////////////////////////
+	//Texture Compression Begin
+
+	/** Return false if the Attribute was not set previously.*/
+	UFUNCTION(BlueprintCallable, Category = "Interchange | Node | Texture")
+	bool GetCustomCompressionNoAlpha(bool& AttributeValue) const
+	{
+		IMPLEMENT_NODE_ATTRIBUTE_GETTER(CompressionNoAlpha, bool);
+	}
+
+	UFUNCTION(BlueprintCallable, Category = "Interchange | Node | Texture")
+	bool SetCustomCompressionNoAlpha(const bool& AttributeValue, bool bAddApplyDelegate = true)
+	{
+#if WITH_EDITORONLY_DATA
+		IMPLEMENT_NODE_ATTRIBUTE_SETTER(UInterchangeTextureNode, CompressionNoAlpha, bool, UTexture)
+#else
+		IMPLEMENT_NODE_ATTRIBUTE_SETTER_NODELEGATE(CompressionNoAlpha, bool)
+#endif
+	}
+
+	/** Return false if the Attribute was not set previously.*/
+	UFUNCTION(BlueprintCallable, Category = "Interchange | Node | Texture")
+	bool GetCustomDeferCompression(bool& AttributeValue) const
+	{
+		IMPLEMENT_NODE_ATTRIBUTE_GETTER(DeferCompression, bool);
+	}
+
+	UFUNCTION(BlueprintCallable, Category = "Interchange | Node | Texture")
+	bool SetCustomDeferCompression(const bool& AttributeValue, bool bAddApplyDelegate = true)
+	{
+#if WITH_EDITORONLY_DATA
+		IMPLEMENT_NODE_ATTRIBUTE_SETTER(UInterchangeTextureNode, DeferCompression, bool, UTexture)
+#else
+		IMPLEMENT_NODE_ATTRIBUTE_SETTER_NODELEGATE(DeferCompression, bool)
+#endif
+	}
+
+	/** Return false if the Attribute was not set previously.*/
+	UFUNCTION(BlueprintCallable, Category = "Interchange | Node | Texture")
+	bool GetCustomLossyCompressionAmount(uint8& AttributeValue) const
+	{
+		IMPLEMENT_NODE_ATTRIBUTE_GETTER(LossyCompressionAmount, uint8);
+	}
+
+	UFUNCTION(BlueprintCallable, Category = "Interchange | Node | Texture")
+	bool SetCustomLossyCompressionAmount(const uint8& AttributeValue, bool bAddApplyDelegate = true)
+	{
+#if WITH_EDITORONLY_DATA
+		IMPLEMENT_NODE_ATTRIBUTE_SETTER(UInterchangeTextureNode, LossyCompressionAmount, uint8, UTexture)
+#else
+		IMPLEMENT_NODE_ATTRIBUTE_SETTER_NODELEGATE(LossyCompressionAmount, uint8)
+#endif
+	}
+
+	/** Return false if the Attribute was not set previously.*/
+	UFUNCTION(BlueprintCallable, Category = "Interchange | Node | Texture")
+	bool GetCustomMaxTextureSize(int32& AttributeValue) const
+	{
+		IMPLEMENT_NODE_ATTRIBUTE_GETTER(MaxTextureSize, int32);
+	}
+
+	UFUNCTION(BlueprintCallable, Category = "Interchange | Node | Texture")
+	bool SetCustomMaxTextureSize(const int32& AttributeValue, bool bAddApplyDelegate = true)
+	{
+#if WITH_EDITORONLY_DATA
+		IMPLEMENT_NODE_ATTRIBUTE_SETTER(UInterchangeTextureNode, MaxTextureSize, int32, UTexture)
+#else
+		IMPLEMENT_NODE_ATTRIBUTE_SETTER_NODELEGATE(MaxTextureSize, int32)
+#endif
+	}
+
+	/** Return false if the Attribute was not set previously.*/
+	UFUNCTION(BlueprintCallable, Category = "Interchange | Node | Texture")
+	bool GetCustomCompressionQuality(uint8& AttributeValue) const
+	{
+		IMPLEMENT_NODE_ATTRIBUTE_GETTER(CompressionQuality, uint8);
+	}
+
+	UFUNCTION(BlueprintCallable, Category = "Interchange | Node | Texture")
+	bool SetCustomCompressionQuality(const uint8& AttributeValue, bool bAddApplyDelegate = true)
+	{
+#if WITH_EDITORONLY_DATA
+		IMPLEMENT_NODE_ATTRIBUTE_SETTER(UInterchangeTextureNode, CompressionQuality, uint8, UTexture)
+#else
+		IMPLEMENT_NODE_ATTRIBUTE_SETTER_NODELEGATE(CompressionQuality, uint8)
+#endif
+	}
+
+	/** Return false if the Attribute was not set previously.*/
+	UFUNCTION(BlueprintCallable, Category = "Interchange | Node | Texture")
+	bool GetCustomCompressionSettings(uint8& AttributeValue) const
+	{
+		IMPLEMENT_NODE_ATTRIBUTE_GETTER(CompressionSettings, uint8);
+	}
+
+	UFUNCTION(BlueprintCallable, Category = "Interchange | Node | Texture")
+	bool SetCustomCompressionSettings(const uint8& AttributeValue, bool bAddApplyDelegate = true)
+	{
+		IMPLEMENT_NODE_ATTRIBUTE_SETTER(UInterchangeTextureNode, CompressionSettings, uint8, UTexture)
+	}
+
+	//////////////////////////////////////////////////////////////////////////
+	//Texture general
+
+	/** Return false if the Attribute was not set previously.*/
+	UFUNCTION(BlueprintCallable, Category = "Interchange | Node | Texture")
+	bool GetCustombDitherMipMapAlpha(bool& AttributeValue) const
+	{
+		IMPLEMENT_NODE_ATTRIBUTE_GETTER(bDitherMipMapAlpha, bool);
+	}
+
+	UFUNCTION(BlueprintCallable, Category = "Interchange | Node | Texture")
+	bool SetCustombDitherMipMapAlpha(const bool& AttributeValue, bool bAddApplyDelegate = true)
+	{
+#if WITH_EDITORONLY_DATA
+		IMPLEMENT_NODE_ATTRIBUTE_SETTER(UInterchangeTextureNode, bDitherMipMapAlpha, bool, UTexture)
+#else
+		IMPLEMENT_NODE_ATTRIBUTE_SETTER_NODELEGATE(bDitherMipMapAlpha, bool)
+#endif
+	}
+
+	/** Return false if the Attribute was not set previously.*/
+	UFUNCTION(BlueprintCallable, Category = "Interchange | Node | Texture")
+	bool GetCustomAlphaCoverageThresholds(FVector4& AttributeValue) const
+	{
+		IMPLEMENT_NODE_ATTRIBUTE_GETTER(AlphaCoverageThresholds, FVector4);
+	}
+
+	UFUNCTION(BlueprintCallable, Category = "Interchange | Node | Texture")
+	bool SetCustomAlphaCoverageThresholds(const FVector4& AttributeValue, bool bAddApplyDelegate = true)
+	{
+#if WITH_EDITORONLY_DATA
+		IMPLEMENT_NODE_ATTRIBUTE_SETTER(UInterchangeTextureNode, AlphaCoverageThresholds, FVector4, UTexture)
+#else
+		IMPLEMENT_NODE_ATTRIBUTE_SETTER_NODELEGATE(AlphaCoverageThresholds, FVector4)
+#endif
+	}
+
+	/** Return false if the Attribute was not set previously.*/
+	UFUNCTION(BlueprintCallable, Category = "Interchange | Node | Texture")
+	bool GetCustombFlipGreenChannel(bool& AttributeValue) const
+	{
+		IMPLEMENT_NODE_ATTRIBUTE_GETTER(bFlipGreenChannel, bool);
+	}
+
+	UFUNCTION(BlueprintCallable, Category = "Interchange | Node | Texture")
+	bool SetCustombFlipGreenChannel(const bool& AttributeValue, bool bAddApplyDelegate = true)
+	{
+#if WITH_EDITORONLY_DATA
+		IMPLEMENT_NODE_ATTRIBUTE_SETTER(UInterchangeTextureNode, bFlipGreenChannel, bool, UTexture)
+#else
+		IMPLEMENT_NODE_ATTRIBUTE_SETTER_NODELEGATE(bFlipGreenChannel, bool)
+#endif
+	}
+
+	/** Return false if the Attribute was not set previously.*/
+	UFUNCTION(BlueprintCallable, Category = "Interchange | Node | Texture")
+	bool GetCustombForcePVRTC4(bool& AttributeValue) const
+	{
+		IMPLEMENT_NODE_ATTRIBUTE_GETTER(bForcePVRTC4, bool);
+	}
+
+	UFUNCTION(BlueprintCallable, Category = "Interchange | Node | Texture")
+	bool SetCustombForcePVRTC4(const bool& AttributeValue, bool bAddApplyDelegate = true)
+	{
+#if WITH_EDITORONLY_DATA
+		IMPLEMENT_NODE_ATTRIBUTE_SETTER(UInterchangeTextureNode, bForcePVRTC4, bool, UTexture)
+#else
+		IMPLEMENT_NODE_ATTRIBUTE_SETTER_NODELEGATE(bForcePVRTC4, bool)
+#endif
+	}
+
+	/** Return false if the Attribute was not set previously.*/
+	UFUNCTION(BlueprintCallable, Category = "Interchange | Node | Texture")
+	bool GetCustomPowerOfTwoMode(uint8& AttributeValue) const
+	{
+		IMPLEMENT_NODE_ATTRIBUTE_GETTER(PowerOfTwoMode, uint8);
+	}
+
+	UFUNCTION(BlueprintCallable, Category = "Interchange | Node | Texture")
+	bool SetCustomPowerOfTwoMode(const uint8& AttributeValue, bool bAddApplyDelegate = true)
+	{
+#if WITH_EDITORONLY_DATA
+		IMPLEMENT_NODE_ATTRIBUTE_SETTER(UInterchangeTextureNode, PowerOfTwoMode, uint8, UTexture)
+#else
+		IMPLEMENT_NODE_ATTRIBUTE_SETTER_NODELEGATE(PowerOfTwoMode, uint8)
+#endif
+	}
+
+	/** Return false if the Attribute was not set previously.*/
+	UFUNCTION(BlueprintCallable, Category = "Interchange | Node | Texture")
+	bool GetCustomPaddingColor(FColor& AttributeValue) const
+	{
+		IMPLEMENT_NODE_ATTRIBUTE_GETTER(PaddingColor, FColor);
+	}
+
+	UFUNCTION(BlueprintCallable, Category = "Interchange | Node | Texture")
+	bool SetCustomPaddingColor(const FColor& AttributeValue, bool bAddApplyDelegate = true)
+	{
+#if WITH_EDITORONLY_DATA
+		IMPLEMENT_NODE_ATTRIBUTE_SETTER(UInterchangeTextureNode, PaddingColor, FColor, UTexture)
+#else
+		IMPLEMENT_NODE_ATTRIBUTE_SETTER_NODELEGATE(PaddingColor, FColor)
+#endif
+	}
+
+	/** Return false if the Attribute was not set previously.*/
+	UFUNCTION(BlueprintCallable, Category = "Interchange | Node | Texture")
+	bool GetCustomFilter(uint8& AttributeValue) const
+	{
+		IMPLEMENT_NODE_ATTRIBUTE_GETTER(Filter, uint8);
+	}
+
+	UFUNCTION(BlueprintCallable, Category = "Interchange | Node | Texture")
+	bool SetCustomFilter(const uint8& AttributeValue, bool bAddApplyDelegate = true)
+	{
+		IMPLEMENT_NODE_ATTRIBUTE_SETTER(UInterchangeTextureNode, Filter, uint8, UTexture)
+	}
+	
+	/** Return false if the Attribute was not set previously.*/
+	UFUNCTION(BlueprintCallable, Category = "Interchange | Node | Texture")
+	bool GetCustomMipLoadOptions(uint8& AttributeValue) const
+	{
+		IMPLEMENT_NODE_ATTRIBUTE_GETTER(MipLoadOptions, uint8);
+	}
+
+	UFUNCTION(BlueprintCallable, Category = "Interchange | Node | Texture")
+	bool SetCustomMipLoadOptions(const uint8& AttributeValue, bool bAddApplyDelegate = true)
+	{
+		IMPLEMENT_NODE_ATTRIBUTE_SETTER(UInterchangeTextureNode, MipLoadOptions, uint8, UTexture)
+	}
+	
+	/** Return false if the Attribute was not set previously.*/
+	UFUNCTION(BlueprintCallable, Category = "Interchange | Node | Texture")
+	bool GetCustomSRGB(bool& AttributeValue) const
+	{
+		IMPLEMENT_NODE_ATTRIBUTE_GETTER(SRGB, bool);
+	}
+
+	UFUNCTION(BlueprintCallable, Category = "Interchange | Node | Texture")
+	bool SetCustomSRGB(const bool& AttributeValue, bool bAddApplyDelegate = true)
+	{
+		IMPLEMENT_NODE_ATTRIBUTE_SETTER(UInterchangeTextureNode, SRGB, bool, UTexture)
+	}
+
+	/** Return false if the Attribute was not set previously.*/
+	UFUNCTION(BlueprintCallable, Category = "Interchange | Node | Texture")
+	bool GetCustombUseLegacyGamma(bool& AttributeValue) const
+	{
+		IMPLEMENT_NODE_ATTRIBUTE_GETTER(bUseLegacyGamma, bool);
+	}
+
+	UFUNCTION(BlueprintCallable, Category = "Interchange | Node | Texture")
+	bool SetCustombUseLegacyGamma(const bool& AttributeValue, bool bAddApplyDelegate = true)
+	{
+#if WITH_EDITORONLY_DATA
+		IMPLEMENT_NODE_ATTRIBUTE_SETTER(UInterchangeTextureNode, bUseLegacyGamma, bool, UTexture)
+#else
+		IMPLEMENT_NODE_ATTRIBUTE_SETTER_NODELEGATE(bUseLegacyGamma, bool)
+#endif
+	}
+
+	/** Return false if the Attribute was not set previously.*/
+	UFUNCTION(BlueprintCallable, Category = "Interchange | Node | Texture")
+	bool GetCustomVirtualTextureStreaming(bool& AttributeValue) const
+	{
+		IMPLEMENT_NODE_ATTRIBUTE_GETTER(VirtualTextureStreaming, bool);
+	}
+
+	UFUNCTION(BlueprintCallable, Category = "Interchange | Node | Texture")
+	bool SetCustomVirtualTextureStreaming(const bool& AttributeValue, bool bAddApplyDelegate = true)
+	{
+		IMPLEMENT_NODE_ATTRIBUTE_SETTER(UInterchangeTextureNode, VirtualTextureStreaming, bool, UTexture)
+	}
+
+	/** Return false if the Attribute was not set previously.*/
+	UFUNCTION(BlueprintCallable, Category = "Interchange | Node | Texture")
+	bool GetCustomAddressX(uint8& AttributeValue) const
+	{
+		IMPLEMENT_NODE_ATTRIBUTE_GETTER(AddressX, uint8);
+	}
+
+	UFUNCTION(BlueprintCallable, Category = "Interchange | Node | Texture")
+	bool SetCustomAddressX(const uint8& AttributeValue, bool bAddApplyDelegate = true)
+	{
+		IMPLEMENT_NODE_ATTRIBUTE_SETTER(UInterchangeTextureNode, AddressX, uint8, UTexture2D)
+	}
+
+	/** Return false if the Attribute was not set previously.*/
+	UFUNCTION(BlueprintCallable, Category = "Interchange | Node | Texture")
+	bool GetCustomAddressY(uint8& AttributeValue) const
+	{
+		IMPLEMENT_NODE_ATTRIBUTE_GETTER(AddressY, uint8);
+	}
+
+	UFUNCTION(BlueprintCallable, Category = "Interchange | Node | Texture")
+	bool SetCustomAddressY(const uint8& AttributeValue, bool bAddApplyDelegate = true)
+	{
+		IMPLEMENT_NODE_ATTRIBUTE_SETTER(UInterchangeTextureNode, AddressY, uint8, UTexture2D)
+	}
+
+	//////////////////////////////////////////////////////////////////////////
+	//Level of Detail
+
+	/** Return false if the Attribute was not set previously.*/
+	UFUNCTION(BlueprintCallable, Category = "Interchange | Node | Texture")
+	bool GetCustombPreserveBorder(bool& AttributeValue) const
+	{
+		IMPLEMENT_NODE_ATTRIBUTE_GETTER(bPreserveBorder, bool);
+	}
+
+	UFUNCTION(BlueprintCallable, Category = "Interchange | Node | Texture")
+	bool SetCustombPreserveBorder(const bool& AttributeValue, bool bAddApplyDelegate = true)
+	{
+#if WITH_EDITORONLY_DATA
+		IMPLEMENT_NODE_ATTRIBUTE_SETTER(UInterchangeTextureNode, bPreserveBorder, bool, UTexture)
+#else
+		IMPLEMENT_NODE_ATTRIBUTE_SETTER_NODELEGATE(bPreserveBorder, bool)
+#endif
+	}
+	
+	/** Return false if the Attribute was not set previously.*/
+	UFUNCTION(BlueprintCallable, Category = "Interchange | Node | Texture")
+	bool GetCustomMipGenSettings(uint8& AttributeValue) const
+	{
+		IMPLEMENT_NODE_ATTRIBUTE_GETTER(MipGenSettings, uint8);
+	}
+
+	UFUNCTION(BlueprintCallable, Category = "Interchange | Node | Texture")
+	bool SetCustomMipGenSettings(const uint8& AttributeValue, bool bAddApplyDelegate = true)
+	{
+#if WITH_EDITORONLY_DATA
+		IMPLEMENT_NODE_ATTRIBUTE_SETTER(UInterchangeTextureNode, MipGenSettings, uint8, UTexture)
+#else
+		IMPLEMENT_NODE_ATTRIBUTE_SETTER_NODELEGATE(MipGenSettings, uint8)
+#endif
+	}
+
+	/** Return false if the Attribute was not set previously.*/
+	UFUNCTION(BlueprintCallable, Category = "Interchange | Node | Texture")
+	bool GetCustomLODBias(int32& AttributeValue) const
+	{
+		IMPLEMENT_NODE_ATTRIBUTE_GETTER(LODBias, int32);
+	}
+
+	UFUNCTION(BlueprintCallable, Category = "Interchange | Node | Texture")
+	bool SetCustomLODBias(const int32& AttributeValue, bool bAddApplyDelegate = true)
+	{
+		IMPLEMENT_NODE_ATTRIBUTE_SETTER(UInterchangeTextureNode, LODBias, int32, UTexture)
+	}
+
+	/** Return false if the Attribute was not set previously.*/
+	UFUNCTION(BlueprintCallable, Category = "Interchange | Node | Texture")
+	bool GetCustomLODGroup(uint8& AttributeValue) const
+	{
+		IMPLEMENT_NODE_ATTRIBUTE_GETTER(LODGroup, uint8);
+	}
+
+	UFUNCTION(BlueprintCallable, Category = "Interchange | Node | Texture")
+	bool SetCustomLODGroup(const uint8& AttributeValue, bool bAddApplyDelegate = true)
+	{
+		IMPLEMENT_NODE_ATTRIBUTE_SETTER(UInterchangeTextureNode, LODGroup, uint8, UTexture)
+	}
+
+	//TODO support per platform data in the FAttributeStorage, so we can set different value per platform at the pipeline stage, We set only the default value for now
+	/** Return false if the Attribute was not set previously.*/
+	UFUNCTION(BlueprintCallable, Category = "Interchange | Node | Texture")
+	bool GetCustomDownscale(float& AttributeValue) const
+	{
+		IMPLEMENT_NODE_ATTRIBUTE_GETTER(Downscale, float);
+	}
+
+	UFUNCTION(BlueprintCallable, Category = "Interchange | Node | Texture")
+	bool SetCustomDownscale(const float& AttributeValue, bool bAddApplyDelegate = true)
+	{
+		IMPLEMENT_NODE_ATTRIBUTE_SETTER(UInterchangeTextureNode, Downscale, float, UTexture)
+	}
+
+	/** Return false if the Attribute was not set previously.*/
+	UFUNCTION(BlueprintCallable, Category = "Interchange | Node | Texture")
+		bool GetCustomDownscaleOptions(uint8& AttributeValue) const
+	{
+		IMPLEMENT_NODE_ATTRIBUTE_GETTER(DownscaleOptions, uint8);
+	}
+
+	UFUNCTION(BlueprintCallable, Category = "Interchange | Node | Texture")
+		bool SetCustomDownscaleOptions(const uint8& AttributeValue, bool bAddApplyDelegate = true)
+	{
+		IMPLEMENT_NODE_ATTRIBUTE_SETTER(UInterchangeTextureNode, DownscaleOptions, uint8, UTexture)
+	}
+
+	//////////////////////////////////////////////////////////////////////////
+	//Compositing
+	
+	/** Return false if the Attribute was not set previously.*/
+	UFUNCTION(BlueprintCallable, Category = "Interchange | Node | Texture")
+	bool GetCustomCompositeTextureMode(uint8& AttributeValue) const
+	{
+		IMPLEMENT_NODE_ATTRIBUTE_GETTER(CompositeTextureMode, uint8);
+	}
+
+	UFUNCTION(BlueprintCallable, Category = "Interchange | Node | Texture")
+	bool SetCustomCompositeTextureMode(const uint8& AttributeValue, bool bAddApplyDelegate = true)
+	{
+#if WITH_EDITORONLY_DATA
+		IMPLEMENT_NODE_ATTRIBUTE_SETTER(UInterchangeTextureNode, CompositeTextureMode, uint8, UTexture)
+#else
+		IMPLEMENT_NODE_ATTRIBUTE_SETTER_NODELEGATE(CompositeTextureMode, uint8)
+#endif
+	}
+	
+	/** Return false if the Attribute was not set previously.*/
+	UFUNCTION(BlueprintCallable, Category = "Interchange | Node | Texture")
+	bool GetCustomCompositePower(float& AttributeValue) const
+	{
+		IMPLEMENT_NODE_ATTRIBUTE_GETTER(CompositePower, float);
+	}
+
+	UFUNCTION(BlueprintCallable, Category = "Interchange | Node | Texture")
+	bool SetCustomCompositePower(const float& AttributeValue, bool bAddApplyDelegate = true)
+	{
+#if WITH_EDITORONLY_DATA
+		IMPLEMENT_NODE_ATTRIBUTE_SETTER(UInterchangeTextureNode, CompositePower, float, UTexture)
+#else
+		IMPLEMENT_NODE_ATTRIBUTE_SETTER_NODELEGATE(CompositePower, float)
+#endif
+	}
+
+private:
+
+	void FillAssetClassFromAttribute()
+	{
+#if WITH_ENGINE
+		FString OperationName = GetTypeName() + TEXT(".GetAssetClassName");
+		FString ClassName;
+		InterchangePrivateNodeBase::GetCustomAttribute<FString>(Attributes, ClassNameAttributeKey, OperationName, ClassName);
+		if (ClassName.Equals(UTexture2D::StaticClass()->GetName()))
+		{
+			AssetClass = UTexture2D::StaticClass();
+			bIsTextureNodeClassInitialized = true;
+		}
+		else if (ClassName.Equals(UTexture::StaticClass()->GetName()))
+		{
+			AssetClass = UTexture::StaticClass();
+			bIsTextureNodeClassInitialized = true;
+		}
+#endif
+	}
+
+	bool SetTextureNodeClassFromClassAttribute()
+	{
+		if (!bIsTextureNodeClassInitialized)
+		{
+			FillAssetClassFromAttribute();
+		}
+		return bIsTextureNodeClassInitialized;
+	}
+
+	bool IsEditorOnlyDataDefined()
+	{
+#if WITH_EDITORONLY_DATA
+		return true;
+#else
+		return false;
+#endif
+	}
+
+	const Interchange::FAttributeKey ClassNameAttributeKey = Interchange::FAttributeKey(TEXT("__ClassTypeAttribute__"));
+
+	//Texture Adjustments
+	const Interchange::FAttributeKey Macro_CustomAdjustBrightnessKey = Interchange::FAttributeKey(TEXT("AdjustBrightness"));
+	const Interchange::FAttributeKey Macro_CustomAdjustBrightnessCurveKey = Interchange::FAttributeKey(TEXT("AdjustBrightnessCurve"));
+	const Interchange::FAttributeKey Macro_CustomAdjustVibranceKey = Interchange::FAttributeKey(TEXT("AdjustVibrance"));
+	const Interchange::FAttributeKey Macro_CustomAdjustSaturationKey = Interchange::FAttributeKey(TEXT("AdjustSaturation"));
+	const Interchange::FAttributeKey Macro_CustomAdjustRGBCurveKey = Interchange::FAttributeKey(TEXT("AdjustRGBCurve"));
+	const Interchange::FAttributeKey Macro_CustomAdjustHueKey = Interchange::FAttributeKey(TEXT("AdjustHue"));
+	const Interchange::FAttributeKey Macro_CustomAdjustMinAlphaKey = Interchange::FAttributeKey(TEXT("AdjustMinAlpha"));
+	const Interchange::FAttributeKey Macro_CustomAdjustMaxAlphaKey = Interchange::FAttributeKey(TEXT("AdjustMaxAlpha"));
+	const Interchange::FAttributeKey Macro_CustombChromaKeyTextureKey = Interchange::FAttributeKey(TEXT("bChromaKeyTexture"));
+	const Interchange::FAttributeKey Macro_CustomChromaKeyThresholdKey = Interchange::FAttributeKey(TEXT("ChromaKeyThreshold"));
+	const Interchange::FAttributeKey Macro_CustomChromaKeyColorKey = Interchange::FAttributeKey(TEXT("ChromaKeyColor"));
+
+	//Texture Compression
+	const Interchange::FAttributeKey Macro_CustomCompressionNoAlphaKey = Interchange::FAttributeKey(TEXT("CompressionNoAlpha"));
+	const Interchange::FAttributeKey Macro_CustomDeferCompressionKey = Interchange::FAttributeKey(TEXT("DeferCompression"));
+	const Interchange::FAttributeKey Macro_CustomLossyCompressionAmountKey = Interchange::FAttributeKey(TEXT("LossyCompressionAmount"));
+	const Interchange::FAttributeKey Macro_CustomMaxTextureSizeKey = Interchange::FAttributeKey(TEXT("MaxTextureSize"));
+	const Interchange::FAttributeKey Macro_CustomCompressionQualityKey = Interchange::FAttributeKey(TEXT("CompressionQuality"));
+	const Interchange::FAttributeKey Macro_CustomCompressionSettingsKey = Interchange::FAttributeKey(TEXT("CompressionSettings"));
+
+	//Texture general
+	const Interchange::FAttributeKey Macro_CustombDitherMipMapAlphaKey = Interchange::FAttributeKey(TEXT("bDitherMipMapAlpha"));
+	const Interchange::FAttributeKey Macro_CustomAlphaCoverageThresholdsKey = Interchange::FAttributeKey(TEXT("AlphaCoverageThresholds"));
+	const Interchange::FAttributeKey Macro_CustombFlipGreenChannelKey = Interchange::FAttributeKey(TEXT("bFlipGreenChannel"));
+	const Interchange::FAttributeKey Macro_CustombForcePVRTC4Key = Interchange::FAttributeKey(TEXT("bForcePVRTC4"));
+	const Interchange::FAttributeKey Macro_CustomPowerOfTwoModeKey = Interchange::FAttributeKey(TEXT("PowerOfTwoMode"));
+	const Interchange::FAttributeKey Macro_CustomPaddingColorKey = Interchange::FAttributeKey(TEXT("PaddingColor"));
+	const Interchange::FAttributeKey Macro_CustomFilterKey = Interchange::FAttributeKey(TEXT("Filter"));
+	const Interchange::FAttributeKey Macro_CustomMipLoadOptionsKey = Interchange::FAttributeKey(TEXT("MipLoadOptions"));
+	const Interchange::FAttributeKey Macro_CustomSRGBKey = Interchange::FAttributeKey(TEXT("SRGB"));
+	const Interchange::FAttributeKey Macro_CustombUseLegacyGammaKey = Interchange::FAttributeKey(TEXT("bUseLegacyGamma"));
+	const Interchange::FAttributeKey Macro_CustomVirtualTextureStreamingKey = Interchange::FAttributeKey(TEXT("VirtualTextureStreaming"));
+	const Interchange::FAttributeKey Macro_CustomAddressXKey = Interchange::FAttributeKey(TEXT("AddressX"));
+	const Interchange::FAttributeKey Macro_CustomAddressYKey = Interchange::FAttributeKey(TEXT("AddressY"));
+
+	//Level of Detail
+	const Interchange::FAttributeKey Macro_CustombPreserveBorderKey = Interchange::FAttributeKey(TEXT("bPreserveBorder"));
+	const Interchange::FAttributeKey Macro_CustomMipGenSettingsKey = Interchange::FAttributeKey(TEXT("MipGenSettings"));
+	const Interchange::FAttributeKey Macro_CustomLODBiasKey = Interchange::FAttributeKey(TEXT("LODBias"));
+	const Interchange::FAttributeKey Macro_CustomLODGroupKey = Interchange::FAttributeKey(TEXT("LODGroup"));
+	//TODO support per platform data in the FAttributeStorage, so we can set different value per platform at the pipeline stage, We set only the default value for now
+	const Interchange::FAttributeKey Macro_CustomDownscaleKey = Interchange::FAttributeKey(TEXT("Downscale"));
+	const Interchange::FAttributeKey Macro_CustomDownscaleOptionsKey = Interchange::FAttributeKey(TEXT("DownscaleOptions"));
+
+	//Compositing
+	const Interchange::FAttributeKey Macro_CustomCompositeTextureModeKey = Interchange::FAttributeKey(TEXT("CompositeTextureMode"));
+	const Interchange::FAttributeKey Macro_CustomCompositePowerKey = Interchange::FAttributeKey(TEXT("CompositePower"));
+
+#if WITH_EDITORONLY_DATA
+	IMPLEMENT_NODE_ATTRIBUTE_APPLY_UOBJECT(AdjustBrightness, float, UTexture, );
+	IMPLEMENT_NODE_ATTRIBUTE_APPLY_UOBJECT(AdjustBrightnessCurve, float, UTexture, );
+	IMPLEMENT_NODE_ATTRIBUTE_APPLY_UOBJECT(AdjustVibrance, float, UTexture, );
+	IMPLEMENT_NODE_ATTRIBUTE_APPLY_UOBJECT(AdjustSaturation, float, UTexture, );
+	IMPLEMENT_NODE_ATTRIBUTE_APPLY_UOBJECT(AdjustRGBCurve, float, UTexture, );
+	IMPLEMENT_NODE_ATTRIBUTE_APPLY_UOBJECT(AdjustHue, float, UTexture, );
+	IMPLEMENT_NODE_ATTRIBUTE_APPLY_UOBJECT(AdjustMinAlpha, float, UTexture, );
+	IMPLEMENT_NODE_ATTRIBUTE_APPLY_UOBJECT(AdjustMaxAlpha, float, UTexture, );
+	IMPLEMENT_NODE_ATTRIBUTE_APPLY_UOBJECT(bChromaKeyTexture, bool, UTexture, );
+	IMPLEMENT_NODE_ATTRIBUTE_APPLY_UOBJECT(ChromaKeyThreshold, float, UTexture, );
+	IMPLEMENT_NODE_ATTRIBUTE_APPLY_UOBJECT(ChromaKeyColor, FColor, UTexture, );
+
+	//Texture Compression
+	IMPLEMENT_NODE_ATTRIBUTE_APPLY_UOBJECT(CompressionNoAlpha, bool, UTexture, );
+	IMPLEMENT_NODE_ATTRIBUTE_APPLY_UOBJECT(DeferCompression, bool, UTexture, );
+	IMPLEMENT_NODE_ATTRIBUTE_APPLY_UOBJECT(LossyCompressionAmount, uint8, UTexture, TEnumAsByte<ETextureLossyCompressionAmount>);
+	IMPLEMENT_NODE_ATTRIBUTE_APPLY_UOBJECT(MaxTextureSize, int32, UTexture, );
+	IMPLEMENT_NODE_ATTRIBUTE_APPLY_UOBJECT(CompressionQuality, uint8, UTexture, TEnumAsByte<enum ETextureCompressionQuality>);
+#endif
+	IMPLEMENT_NODE_ATTRIBUTE_APPLY_UOBJECT(CompressionSettings, uint8, UTexture, TEnumAsByte<enum TextureCompressionSettings>);
+#if WITH_EDITORONLY_DATA
+	//Texture general
+	IMPLEMENT_NODE_ATTRIBUTE_APPLY_UOBJECT(bDitherMipMapAlpha, bool, UTexture, );
+	IMPLEMENT_NODE_ATTRIBUTE_APPLY_UOBJECT(AlphaCoverageThresholds, FVector4, UTexture, );
+	IMPLEMENT_NODE_ATTRIBUTE_APPLY_UOBJECT(bFlipGreenChannel, bool, UTexture, );
+	IMPLEMENT_NODE_ATTRIBUTE_APPLY_UOBJECT(bForcePVRTC4, bool, UTexture, );
+	IMPLEMENT_NODE_ATTRIBUTE_APPLY_UOBJECT(PowerOfTwoMode, uint8, UTexture, TEnumAsByte<enum ETexturePowerOfTwoSetting::Type>);
+	IMPLEMENT_NODE_ATTRIBUTE_APPLY_UOBJECT(PaddingColor, FColor, UTexture, );
+#endif //#if WITH_EDITORONLY_DATA
+	IMPLEMENT_NODE_ATTRIBUTE_APPLY_UOBJECT(Filter, uint8, UTexture, TEnumAsByte<enum TextureFilter>);
+	IMPLEMENT_NODE_ATTRIBUTE_APPLY_UOBJECT(MipLoadOptions, uint8, UTexture, ETextureMipLoadOptions);
+	IMPLEMENT_NODE_ATTRIBUTE_APPLY_UOBJECT(SRGB, bool, UTexture, );
+#if WITH_EDITORONLY_DATA
+	IMPLEMENT_NODE_ATTRIBUTE_APPLY_UOBJECT(bUseLegacyGamma, bool, UTexture, );
+#endif //#if WITH_EDITORONLY_DATA
+	IMPLEMENT_NODE_ATTRIBUTE_APPLY_UOBJECT(VirtualTextureStreaming, bool, UTexture, );
+	IMPLEMENT_NODE_ATTRIBUTE_APPLY_UOBJECT(AddressX, uint8, UTexture2D, TEnumAsByte<enum TextureAddress>);
+	IMPLEMENT_NODE_ATTRIBUTE_APPLY_UOBJECT(AddressY, uint8, UTexture2D, TEnumAsByte<enum TextureAddress>);
+#if WITH_EDITORONLY_DATA
+	//Level of Detail
+	IMPLEMENT_NODE_ATTRIBUTE_APPLY_UOBJECT(bPreserveBorder, bool, UTexture, );
+	IMPLEMENT_NODE_ATTRIBUTE_APPLY_UOBJECT(MipGenSettings, uint8, UTexture, TEnumAsByte<enum TextureMipGenSettings>);
+#endif //#if WITH_EDITORONLY_DATA
+	IMPLEMENT_NODE_ATTRIBUTE_APPLY_UOBJECT(LODBias, int32, UTexture, );
+	IMPLEMENT_NODE_ATTRIBUTE_APPLY_UOBJECT(LODGroup, uint8, UTexture, TEnumAsByte<enum TextureGroup>);
+	//TODO support per platform data in the FAttributeStorage, so we can set different value per platform at the pipeline stage, We set only the default value for now
+	IMPLEMENT_NODE_ATTRIBUTE_APPLY_UOBJECT(Downscale, float, UTexture, );
+	IMPLEMENT_NODE_ATTRIBUTE_APPLY_UOBJECT(DownscaleOptions, uint8, UTexture, ETextureDownscaleOptions);
+#if WITH_EDITORONLY_DATA
+	//Compositing
+	IMPLEMENT_NODE_ATTRIBUTE_APPLY_UOBJECT(CompositeTextureMode, uint8, UTexture, TEnumAsByte<enum ECompositeTextureMode>);
+	IMPLEMENT_NODE_ATTRIBUTE_APPLY_UOBJECT(CompositePower, float, UTexture, );
+#endif //#if WITH_EDITORONLY_DATA
+protected:
+#if WITH_ENGINE
+	TSubclassOf<UTexture> AssetClass = nullptr;
+#endif
+	bool bIsTextureNodeClassInitialized = false;
+};
diff --git a/Engine/Plugins/Runtime/Interchange/Source/Nodes/Public/TextureNode.h b/Engine/Plugins/Runtime/Interchange/Source/Nodes/Public/TextureNode.h
deleted file mode 100644
diff --git a/Engine/Plugins/Runtime/Interchange/Source/Parsers/Fbx/InterchangeFbxParser.Build.cs b/Engine/Plugins/Runtime/Interchange/Source/Parsers/Fbx/InterchangeFbxParser.Build.cs
--- a/Engine/Plugins/Runtime/Interchange/Source/Parsers/Fbx/InterchangeFbxParser.Build.cs
+++ b/Engine/Plugins/Runtime/Interchange/Source/Parsers/Fbx/InterchangeFbxParser.Build.cs
@@ -24,6 +24,12 @@ namespace UnrealBuildTool.Rules
 				}
 			);
 
+			if (Target.bCompileAgainstEngine)
+			{
+				PrivateDependencyModuleNames.Add("Engine");
+
+			}
+
 			AddEngineThirdPartyPrivateStaticDependencies(Target,
 				"FBX"
 			);
diff --git a/Engine/Plugins/Runtime/Interchange/Source/Parsers/Fbx/Private/FbxAPI.cpp b/Engine/Plugins/Runtime/Interchange/Source/Parsers/Fbx/Private/FbxAPI.cpp
new file mode 100644
--- /dev/null
+++ b/Engine/Plugins/Runtime/Interchange/Source/Parsers/Fbx/Private/FbxAPI.cpp
@@ -0,0 +1,89 @@
+// Copyright Epic Games, Inc. All Rights Reserved.
+
+#include "FbxAPI.h"
+
+#include "CoreMinimal.h"
+#include "FbxConvert.h"
+#include "FbxInclude.h"
+#include "FbxMaterial.h"
+#include "FbxScene.h"
+#include "InterchangeMaterialNode.h"
+#include "InterchangeTextureNode.h"
+#include "Nodes/InterchangeBaseNodeContainer.h"
+
+namespace UE
+{
+	namespace FbxParser
+	{
+		namespace Private
+		{
+			bool FFbxParser::LoadFbxFile(const FString& Filename, TArray<FString>& JSonErrorMessages)
+			{
+				int32 SDKMajor, SDKMinor, SDKRevision;
+				
+				//The first thing to do is to create the FBX Manager which is the object allocator for almost all the classes in the SDK
+				SDKManager = FbxManager::Create();
+				if (!SDKManager)
+				{
+					JSonErrorMessages.Add(TEXT("{\"Msg\" : {\"Type\" : \"Error\",\n\"Msg\" : \"Cannot create fbx sdk manager\"}}"));
+					return false;
+				}
+
+				//Create an IOSettings object. This object holds all import/export settings.
+				FbxIOSettings* ios = FbxIOSettings::Create(SDKManager, IOSROOT);
+				ios->SetBoolProp(IMP_FBX_MATERIAL, true);
+				ios->SetBoolProp(IMP_FBX_TEXTURE, true);
+				ios->SetBoolProp(IMP_FBX_LINK, true);
+				ios->SetBoolProp(IMP_FBX_SHAPE, true);
+				ios->SetBoolProp(IMP_FBX_GOBO, true);
+				ios->SetBoolProp(IMP_FBX_ANIMATION, true);
+				ios->SetBoolProp(IMP_SKINS, true);
+				ios->SetBoolProp(IMP_DEFORMATION, true);
+				ios->SetBoolProp(IMP_FBX_GLOBAL_SETTINGS, true);
+				ios->SetBoolProp(IMP_TAKE, true);
+				SDKManager->SetIOSettings(ios);
+
+				//Create an FBX scene. This object holds most objects imported/exported from/to files.
+				SDKScene = FbxScene::Create(SDKManager, "My Scene");
+				if (!SDKScene)
+				{
+					JSonErrorMessages.Add(TEXT("{\"Msg\" : {\"Type\" : \"Error\",\n\"Msg\" : \"Cannot create fbx sdk scene\"}}"));
+					return false;
+				}
+				
+				// Create an importer.
+				SDKImporter = FbxImporter::Create(SDKManager, "");
+
+				// Get the version number of the FBX files generated by the
+				// version of FBX SDK that you are using.
+				FbxManager::GetFileFormatVersion(SDKMajor, SDKMinor, SDKRevision);
+
+				// Initialize the importer by providing a filename.
+				const bool bImportStatus = SDKImporter->Initialize(TCHAR_TO_UTF8(*Filename));
+				if (!bImportStatus)
+				{
+					JSonErrorMessages.Add(TEXT("{\"Msg\" : {\"Type\" : \"Error\",\n\"Msg\" : \"Cannot open fbx filename\"}}"));
+					return false;
+				}
+
+				bool bStatus = SDKImporter->Import(SDKScene);
+
+				const FbxGlobalSettings& GlobalSettings = SDKScene->GetGlobalSettings();
+				FbxTime::EMode TimeMode = GlobalSettings.GetTimeMode();
+				//Set the original framerate from the current fbx file
+				float FbxFramerate = FbxTime::GetFrameRate(TimeMode);
+
+				//We always convert scene to UE axis and units
+				FFbxConvert::ConvertScene(SDKScene);
+
+				return true;
+			}
+
+			void FFbxParser::FillContainerWithFbxScene(UInterchangeBaseNodeContainer& NodeContainer, TArray<FString>& JSonErrorMessages)
+			{
+				FFbxMaterial::AddAllSceneMaterials(SDKScene, NodeContainer, JSonErrorMessages);
+				FFbxScene::AddHierarchy(SDKScene, NodeContainer, JSonErrorMessages);
+			}
+		} //ns Private
+	} //ns FbxParser
+} //ns UE
diff --git a/Engine/Plugins/Runtime/Interchange/Source/Parsers/Fbx/Private/FbxAPI.h b/Engine/Plugins/Runtime/Interchange/Source/Parsers/Fbx/Private/FbxAPI.h
new file mode 100644
--- /dev/null
+++ b/Engine/Plugins/Runtime/Interchange/Source/Parsers/Fbx/Private/FbxAPI.h
@@ -0,0 +1,36 @@
+// Copyright Epic Games, Inc. All Rights Reserved.
+
+#pragma once
+
+#include "CoreMinimal.h"
+#include "FbxInclude.h"
+
+#define FBX_METADATA_PREFIX TEXT("FBX.")
+#define INVALID_UNIQUE_ID 0xFFFFFFFFFFFFFFFF
+
+class UInterchangeBaseNodeContainer;
+
+namespace UE
+{
+	namespace FbxParser
+	{
+		namespace Private
+		{
+			class FFbxParser
+			{
+			public:
+				/* Load an fbx file into the fbx sdk, return false if the file could not be load. */
+				bool LoadFbxFile(const FString& Filename, TArray<FString>& JsonErrorMessages);
+
+				/* Extract the fbx data from the sdk into our node container */
+				void FillContainerWithFbxScene(UInterchangeBaseNodeContainer& NodeContainer, TArray<FString>& JsonErrorMessages);
+			
+			private:
+				FbxManager* SDKManager = nullptr;
+				FbxScene* SDKScene = nullptr;
+				FbxImporter* SDKImporter = nullptr;
+				//FbxGeometryConverter* SDKGeometryConverter = nullptr;
+			};
+		}
+	}
+}
diff --git a/Engine/Plugins/Runtime/Interchange/Source/Parsers/Fbx/Private/FbxConvert.cpp b/Engine/Plugins/Runtime/Interchange/Source/Parsers/Fbx/Private/FbxConvert.cpp
new file mode 100644
--- /dev/null
+++ b/Engine/Plugins/Runtime/Interchange/Source/Parsers/Fbx/Private/FbxConvert.cpp
@@ -0,0 +1,179 @@
+// Copyright Epic Games, Inc. All Rights Reserved.
+
+#include "FbxConvert.h"
+
+#include "CoreMinimal.h"
+#include "FbxInclude.h"
+
+namespace UE
+{
+	namespace FbxParser
+	{
+		namespace Private
+		{
+			void FFbxConvert::ConvertScene(FbxScene* SDKScene)
+			{
+				if (!ensure(SDKScene))
+				{
+					//Cannot convert a null scene
+					return;
+				}
+
+				const FbxGlobalSettings& GlobalSettings = SDKScene->GetGlobalSettings();
+				FbxTime::EMode TimeMode = GlobalSettings.GetTimeMode();
+				//Set the original framerate from the current fbx file
+				float FbxFramerate = FbxTime::GetFrameRate(TimeMode);
+
+				//Merge the anim stack before the conversion since the above 0 layer will not be converted
+				int32 AnimStackCount = SDKScene->GetSrcObjectCount<FbxAnimStack>();
+				//Merge the animation stack layer before converting the scene
+				for (int32 AnimStackIndex = 0; AnimStackIndex < AnimStackCount; AnimStackIndex++)
+				{
+					FbxAnimStack* CurAnimStack = SDKScene->GetSrcObject<FbxAnimStack>(AnimStackIndex);
+					if (CurAnimStack->GetMemberCount() > 1)
+					{
+						MergeAllLayerAnimation(SDKScene, CurAnimStack, FbxFramerate);
+					}
+				}
+
+				//Set the original file information
+				FbxAxisSystem FileAxisSystem = SDKScene->GetGlobalSettings().GetAxisSystem();
+				FbxSystemUnit FileUnitSystem = SDKScene->GetGlobalSettings().GetSystemUnit();
+
+
+				//UE is: z up, front x, left handed
+				FbxAxisSystem::EUpVector UpVector = FbxAxisSystem::eZAxis;
+				FbxAxisSystem::EFrontVector FrontVector = (FbxAxisSystem::EFrontVector)FbxAxisSystem::eParityEven;
+				FbxAxisSystem::ECoordSystem CoordSystem = FbxAxisSystem::eLeftHanded;
+				FbxAxisSystem UnrealImportAxis(UpVector, FrontVector, CoordSystem);
+
+				if (FileAxisSystem != UnrealImportAxis)
+				{
+					FbxRootNodeUtility::RemoveAllFbxRoots(SDKScene);
+					UnrealImportAxis.ConvertScene(SDKScene);
+				}
+
+				if (FileUnitSystem != FbxSystemUnit::cm)
+				{
+					FbxSystemUnit::cm.ConvertScene(SDKScene);
+				}
+
+				//Reset all the transform evaluation cache since we change some node transform
+				SDKScene->GetAnimationEvaluator()->Reset();
+			}
+
+			FTransform FFbxConvert::GetTransform(const FbxAMatrix& Matrix)
+			{
+				FTransform Transform = FTransform::Identity;
+				FbxVector4 FbxLocalT = Matrix.GetT();
+				FbxVector4 FbxLocalS = Matrix.GetS();
+				FbxQuaternion FbxLocalQ = Matrix.GetQ();
+				FVector LocalT = FVector(FbxLocalT[0], FbxLocalT[1], FbxLocalT[2]);
+				Transform.SetTranslation(LocalT);
+				FVector LocalS = FVector(FbxLocalS[0], FbxLocalS[1], FbxLocalS[2]);
+				Transform.SetScale3D(LocalS);
+				FQuat LocalQ;
+				LocalQ.X = FbxLocalQ[0];
+				LocalQ.Y = FbxLocalQ[1];
+				LocalQ.Z = FbxLocalQ[2];
+				LocalQ.W = FbxLocalQ[3];
+				Transform.SetRotation(LocalQ);
+				return Transform;
+			}
+
+			ANSICHAR* FFbxConvert::MakeName(const ANSICHAR* Name)
+			{
+				const int SpecialChars[] = { '.', ',', '/', '`', '%' };
+
+				const int len = FCStringAnsi::Strlen(Name);
+				ANSICHAR* TmpName = new ANSICHAR[len + 1];
+
+				FCStringAnsi::Strcpy(TmpName, len + 1, Name);
+
+				for (int32 i = 0; i < UE_ARRAY_COUNT(SpecialChars); i++)
+				{
+					ANSICHAR* CharPtr = TmpName;
+					while ((CharPtr = FCStringAnsi::Strchr(CharPtr, SpecialChars[i])) != NULL)
+					{
+						CharPtr[0] = '_';
+					}
+				}
+
+				// Remove namespaces
+				ANSICHAR* NewName;
+				NewName = FCStringAnsi::Strchr(TmpName, ':');
+
+				// there may be multiple namespace, so find the last ':'
+				while (NewName && FCStringAnsi::Strchr(NewName + 1, ':'))
+				{
+					NewName = FCStringAnsi::Strchr(NewName + 1, ':');
+				}
+
+				if (NewName)
+				{
+					return NewName + 1;
+				}
+
+				return TmpName;
+			}
+
+			/**
+			 * Convert ANSI char to a FString using ANSI_TO_TCHAR macro
+			 */
+			FString FFbxConvert::MakeString(const ANSICHAR* Name)
+			{
+				return FString(ANSI_TO_TCHAR(Name));
+			}
+
+			void FFbxConvert::ApplyUnroll(FbxNode* Node, FbxAnimLayer* Layer, FbxAnimCurveFilterUnroll* UnrollFilter)
+			{
+				if (!ensure(Node) || !ensure(Layer) || !ensure(UnrollFilter))
+				{
+					return;
+				}
+
+				FbxAnimCurveNode* lCN = Node->LclRotation.GetCurveNode(Layer);
+				if (lCN)
+				{
+					FbxAnimCurve* lRCurve[3];
+					lRCurve[0] = lCN->GetCurve(0);
+					lRCurve[1] = lCN->GetCurve(1);
+					lRCurve[2] = lCN->GetCurve(2);
+
+
+					// Set bone rotation order
+					EFbxRotationOrder RotationOrder = eEulerXYZ;
+					Node->GetRotationOrder(FbxNode::eSourcePivot, RotationOrder);
+					UnrollFilter->SetRotationOrder((FbxEuler::EOrder)(RotationOrder));
+
+					UnrollFilter->Apply(lRCurve, 3);
+				}
+
+				for (int32 i = 0; i < Node->GetChildCount(); i++)
+				{
+					ApplyUnroll(Node->GetChild(i), Layer, UnrollFilter);
+				}
+			}
+
+			void FFbxConvert::MergeAllLayerAnimation(FbxScene* SDKScene, FbxAnimStack* AnimStack, float ResampleRate)
+			{
+				if (!ensure(SDKScene) || !ensure(AnimStack))
+				{
+					return;
+				}
+				FbxTime FramePeriod;
+				FramePeriod.SetSecondDouble(1.0 / (double)ResampleRate);
+
+				FbxTimeSpan TimeSpan = AnimStack->GetLocalTimeSpan();
+				AnimStack->BakeLayers(SDKScene->GetAnimationEvaluator(), TimeSpan.GetStart(), TimeSpan.GetStop(), FramePeriod);
+
+				// always apply unroll filter
+				FbxAnimCurveFilterUnroll UnrollFilter;
+
+				FbxAnimLayer* Layer = AnimStack->GetMember<FbxAnimLayer>(0);
+				UnrollFilter.Reset();
+				ApplyUnroll(SDKScene->GetRootNode(), Layer, &UnrollFilter);
+			}
+		}
+	}
+}
diff --git a/Engine/Plugins/Runtime/Interchange/Source/Parsers/Fbx/Private/FbxConvert.h b/Engine/Plugins/Runtime/Interchange/Source/Parsers/Fbx/Private/FbxConvert.h
new file mode 100644
--- /dev/null
+++ b/Engine/Plugins/Runtime/Interchange/Source/Parsers/Fbx/Private/FbxConvert.h
@@ -0,0 +1,76 @@
+// Copyright Epic Games, Inc. All Rights Reserved.
+
+#pragma once
+
+#include "CoreMinimal.h"
+#include "FbxInclude.h"
+
+namespace UE
+{
+	namespace FbxParser
+	{
+		namespace Private
+		{
+			struct FFbxConvert
+			{
+			public:
+
+				//////////////////////////////////////////////////////////////////////////
+				/** Transform Conversion API Begin */
+
+				/**
+				 * Create a UE transform from a fbx matrix
+				 */
+				static FTransform GetTransform(const FbxAMatrix& Matrix);
+
+				/** Transform Conversion API End */
+				//////////////////////////////////////////////////////////////////////////
+
+				
+				//////////////////////////////////////////////////////////////////////////
+				/** Scene Conversion API Begin */
+
+				/**
+				 * Convert a fbx scene 
+				 */
+				static void ConvertScene(FbxScene* SDKScene);
+				
+				/** Scene Conversion API End */
+				//////////////////////////////////////////////////////////////////////////
+
+
+				//////////////////////////////////////////////////////////////////////////
+				/** String Conversion API Begin */
+
+				/**
+				 * Replace all special characters with '_', then remove all namespace
+				 * Special characters are . , / ` %
+				 */
+				static ANSICHAR* MakeName(const ANSICHAR* Name);
+
+				/**
+				 * Convert ANSI char to a FString using ANSI_TO_TCHAR macro
+				 */
+				static FString MakeString(const ANSICHAR* Name);
+
+				/** String Conversion API End */
+				//////////////////////////////////////////////////////////////////////////
+
+			private:
+				
+				//////////////////////////////////////////////////////////////////////////
+				/** Scene Conversion Private Implementation Begin */
+
+				/**
+				 * The Unroll filter expects only rotation curves, we need to walk the scene and extract the
+				 * rotation curves from the nodes property. This can become time consuming but we have no choice.
+				 */
+				static void ApplyUnroll(FbxNode* pNode, FbxAnimLayer* pLayer, FbxAnimCurveFilterUnroll* pUnrollFilter);
+				static void MergeAllLayerAnimation(FbxScene* SDKScene, FbxAnimStack* AnimStack, float ResampleRate);
+
+				/** Scene Conversion Private Implementation End */
+				//////////////////////////////////////////////////////////////////////////
+			};
+		}
+	}
+}
diff --git a/Engine/Plugins/Runtime/Interchange/Source/Parsers/Fbx/Private/FbxInclude.h b/Engine/Plugins/Runtime/Interchange/Source/Parsers/Fbx/Private/FbxInclude.h
new file mode 100644
--- /dev/null
+++ b/Engine/Plugins/Runtime/Interchange/Source/Parsers/Fbx/Private/FbxInclude.h
@@ -0,0 +1,56 @@
+// Copyright Epic Games, Inc. All Rights Reserved.
+
+#pragma once
+
+#include "CoreMinimal.h"
+
+// Temporarily disable a few warnings due to virtual function abuse in FBX source files
+#pragma warning( push )
+
+#pragma warning( disable : 4263 ) // 'function' : member function does not override any base class virtual member function
+#pragma warning( disable : 4264 ) // 'virtual_function' : no override available for virtual member function from base 'class'; function is hidden
+
+// Include the fbx sdk header
+// temp undef/redef of _O_RDONLY because kfbxcache.h (included by fbxsdk.h) does
+// a weird use of these identifiers inside an enum.
+#ifdef _O_RDONLY
+#define TMP_UNFBX_BACKUP_O_RDONLY _O_RDONLY
+#define TMP_UNFBX_BACKUP_O_WRONLY _O_WRONLY
+#undef _O_RDONLY
+#undef _O_WRONLY
+#endif
+
+//Robert G. : Packing was only set for the 64bits platform, but we also need it for 32bits.
+//This was found while trying to trace a loop that iterate through all character links.
+//The memory didn't match what the debugger displayed, obviously since the packing was not right.
+#pragma pack(push,8)
+
+#if PLATFORM_WINDOWS
+// _CRT_SECURE_NO_DEPRECATE is defined but is not enough to suppress the deprecation
+// warning for vsprintf and stricmp in VS2010.  Since FBX is able to properly handle the non-deprecated
+// versions on the appropriate platforms, _CRT_SECURE_NO_DEPRECATE is temporarily undefined before
+// including the FBX headers
+
+// The following is a hack to make the FBX header files compile correctly under Visual Studio 2012 and Visual Studio 2013
+#if _MSC_VER >= 1700
+#define FBX_DLL_MSC_VER 1600
+#endif
+
+
+#endif // PLATFORM_WINDOWS
+
+// FBX casts null pointer to a reference
+THIRD_PARTY_INCLUDES_START
+#include <fbxsdk.h>
+THIRD_PARTY_INCLUDES_END
+
+#pragma pack(pop)
+
+#ifdef TMP_UNFBX_BACKUP_O_RDONLY
+#define _O_RDONLY TMP_FBX_BACKUP_O_RDONLY
+#define _O_WRONLY TMP_FBX_BACKUP_O_WRONLY
+#undef TMP_UNFBX_BACKUP_O_RDONLY
+#undef TMP_UNFBX_BACKUP_O_WRONLY
+#endif
+
+#pragma warning( pop )
diff --git a/Engine/Plugins/Runtime/Interchange/Source/Parsers/Fbx/Private/FbxMaterial.cpp b/Engine/Plugins/Runtime/Interchange/Source/Parsers/Fbx/Private/FbxMaterial.cpp
new file mode 100644
--- /dev/null
+++ b/Engine/Plugins/Runtime/Interchange/Source/Parsers/Fbx/Private/FbxMaterial.cpp
@@ -0,0 +1,187 @@
+// Copyright Epic Games, Inc. All Rights Reserved.
+
+#include "FbxMaterial.h"
+
+#include "CoreMinimal.h"
+#include "FbxConvert.h"
+#include "FbxInclude.h"
+#include "InterchangeMaterialNode.h"
+#include "InterchangeTextureNode.h"
+#include "Nodes/InterchangeBaseNodeContainer.h"
+
+namespace UE
+{
+	namespace FbxParser
+	{
+		namespace Private
+		{
+			UInterchangeMaterialNode* FFbxMaterial::CreateMaterialNode(UInterchangeBaseNodeContainer& NodeContainer, const FString& NodeName, TArray<FString>& JSonErrorMessages)
+			{
+				FName DisplayLabel = *NodeName;
+				FName NodeUID(*NodeName);
+				UInterchangeMaterialNode* MaterialNode = NewObject<UInterchangeMaterialNode>(&NodeContainer, NAME_None);
+				if (!ensure(MaterialNode))
+				{
+					JSonErrorMessages.Add(TEXT("{\"Msg\" : {\"Type\" : \"Error\",\n\"Msg\" : \"Cannot allocate a node when importing fbx\"}}"));
+					return nullptr;
+				}
+				// Creating a UMaterialInterface
+				MaterialNode->InitializeMaterialNode(NodeUID, DisplayLabel, TEXT("Material"));
+				MaterialNode->SetPayLoadKey(NodeName);
+				NodeContainer.AddNode(MaterialNode);
+				return MaterialNode;
+			}
+
+			UInterchangeTextureNode* FFbxMaterial::CreateTexture2DNode(UInterchangeBaseNodeContainer& NodeContainer, const FString& TextureFilePath, TArray<FString>& JSonErrorMessages)
+			{
+				FName DisplayLabel = *FPaths::GetBaseFilename(TextureFilePath);
+				FName NodeUID(*TextureFilePath);
+				UInterchangeTextureNode* TextureNode = NewObject<UInterchangeTextureNode>(&NodeContainer, NAME_None);
+				if (!ensure(TextureNode))
+				{
+					JSonErrorMessages.Add(TEXT("{\"Msg\" : {\"Type\" : \"Error\",\n\"Msg\" : \"Cannot allocate a node when importing fbx\"}}"));
+					return nullptr;
+				}
+				// Creating a UTexture2D
+				TextureNode->InitializeTextureNode(NodeUID, DisplayLabel, TEXT("Texture2D"));
+				TextureNode->SetPayLoadKey(TextureFilePath);
+				NodeContainer.AddNode(TextureNode);
+				return TextureNode;
+			}
+
+			UInterchangeMaterialNode* FFbxMaterial::AddNodeMaterial(FbxSurfaceMaterial* SurfaceMaterial, UInterchangeBaseNodeContainer& NodeContainer, TArray<FString>& JSonErrorMessages)
+			{
+				//Create a material node
+				FString MaterialName = FFbxConvert::MakeString(FFbxConvert::MakeName(SurfaceMaterial->GetName()));
+				FName NodeUID(*MaterialName);
+				UInterchangeMaterialNode* MaterialNode = Cast<UInterchangeMaterialNode>(NodeContainer.GetNode(NodeUID));
+				if (!MaterialNode)
+				{
+					MaterialNode = CreateMaterialNode(NodeContainer, MaterialName, JSonErrorMessages);
+					if (MaterialNode == nullptr)
+					{
+						JSonErrorMessages.Add(TEXT("{\"Msg\" : {\"Type\" : \"Error\",\n\"Msg\" : \"Cannot create fbx material (" + MaterialName + TEXT(")\"}}")));
+						return nullptr;
+					}
+
+					auto SetMaterialParameter = [&NodeContainer, &MaterialNode, &SurfaceMaterial, &JSonErrorMessages, &MaterialName](const char* FbxMaterialProperty, EInterchangeMaterialNodeParameterName MaterialParameterName)
+					{
+						bool bSetMaterial = false;
+						FbxProperty FbxProperty = SurfaceMaterial->FindProperty(FbxMaterialProperty);
+						if (FbxProperty.IsValid())
+						{
+							int32 UnsupportedTextureCount = FbxProperty.GetSrcObjectCount<FbxLayeredTexture>();
+							UnsupportedTextureCount += FbxProperty.GetSrcObjectCount<FbxProceduralTexture>();
+							int32 TextureCount = FbxProperty.GetSrcObjectCount<FbxTexture>();
+							if (UnsupportedTextureCount > 0)
+							{
+								JSonErrorMessages.Add(TEXT("{\"Msg\" : {\"Type\" : \"Warning\",\n\"Msg\" : \"Layered or procedural Textures are not supported (material" + MaterialName + TEXT(")\"}}")));
+							}
+							else if (TextureCount > 0)
+							{
+								for (int32 TextureIndex = 0; TextureIndex < TextureCount; ++TextureIndex)
+								{
+									FbxFileTexture* FbxTextureFilePath = FbxProperty.GetSrcObject<FbxFileTexture>(TextureIndex);
+									FString TextureFilename = UTF8_TO_TCHAR(FbxTextureFilePath->GetFileName());
+									//Create a texture node and make it child of the material node
+									TArray<FString> JsonErrorMessage;
+									FName NodeUID(*TextureFilename);
+									UInterchangeTextureNode* TextureNode = Cast<UInterchangeTextureNode>(NodeContainer.GetNode(NodeUID));
+									if (!TextureNode)
+									{
+										TextureNode = CreateTexture2DNode(NodeContainer, TextureFilename, JsonErrorMessage);
+									}
+									// add/find UVSet and set it to the texture, we pass index 0 here, I think pipeline should be able to get the UVIndex channel from the name
+									// and modify the parameter to set the correct value.
+// 									FbxString UVSetName = FbxTextureFilePath->UVSet.Get();
+// 									FString LocalUVSetName = UTF8_TO_TCHAR(UVSetName.Buffer());
+									const int32 UVChannelIndex = 0;
+									float ScaleU = (float)FbxTextureFilePath->GetScaleU();
+									float ScaleV = (float)FbxTextureFilePath->GetScaleV();
+									MaterialNode->AddTextureParameterData(MaterialParameterName, TextureNode->GetUniqueID(), UVChannelIndex, ScaleU, ScaleV);
+									MaterialNode->SetDependencyUID(TextureNode->GetUniqueID());
+									bSetMaterial = true;
+								}
+							}
+							else if (MaterialParameterName == EInterchangeMaterialNodeParameterName::BaseColor)
+							{
+								//We support only the base color has a vector color
+								//TODO support all basic attributes has vector or scalar
+								FbxDouble3 FbxColor;
+								bool bFoundColor = true;
+								if (SurfaceMaterial->GetClassId().Is(FbxSurfacePhong::ClassId))
+								{
+									FbxColor = ((FbxSurfacePhong&)(*SurfaceMaterial)).Diffuse.Get();
+								}
+								else if (SurfaceMaterial->GetClassId().Is(FbxSurfaceLambert::ClassId))
+								{
+									FbxColor = ((FbxSurfaceLambert&)(*SurfaceMaterial)).Diffuse.Get();
+								}
+								else
+								{
+									bFoundColor = false;
+								}
+
+								FVector ColorData;
+								if (bFoundColor)
+								{
+									ColorData[0] = (float)(FbxColor[0]);
+									ColorData[1] = (float)(FbxColor[1]);
+									ColorData[2] = (float)(FbxColor[2]);
+								}
+								else
+								{
+									// use random color because there may be multiple materials, then they can be different 
+									ColorData[0] = 0.5f + (0.5f * FMath::Rand()) / RAND_MAX;
+									ColorData[1] = 0.5f + (0.5f * FMath::Rand()) / RAND_MAX;
+									ColorData[2] = 0.5f + (0.5f * FMath::Rand()) / RAND_MAX;
+								}
+								MaterialNode->AddVectorParameterData(MaterialParameterName, ColorData);
+								bSetMaterial = true;
+							}
+						}
+						return bSetMaterial;
+					};
+
+					SetMaterialParameter(FbxSurfaceMaterial::sDiffuse, EInterchangeMaterialNodeParameterName::BaseColor);
+					SetMaterialParameter(FbxSurfaceMaterial::sEmissive, EInterchangeMaterialNodeParameterName::EmissiveColor);
+					SetMaterialParameter(FbxSurfaceMaterial::sSpecular, EInterchangeMaterialNodeParameterName::Specular);
+					SetMaterialParameter(FbxSurfaceMaterial::sSpecularFactor, EInterchangeMaterialNodeParameterName::Roughness);
+					SetMaterialParameter(FbxSurfaceMaterial::sShininess, EInterchangeMaterialNodeParameterName::Metallic);
+					if (!SetMaterialParameter(FbxSurfaceMaterial::sNormalMap, EInterchangeMaterialNodeParameterName::Normal))
+					{
+						SetMaterialParameter(FbxSurfaceMaterial::sBump, EInterchangeMaterialNodeParameterName::Normal);
+					}
+					if (SetMaterialParameter(FbxSurfaceMaterial::sTransparentColor, EInterchangeMaterialNodeParameterName::Opacity))
+					{
+						SetMaterialParameter(FbxSurfaceMaterial::sTransparencyFactor, EInterchangeMaterialNodeParameterName::OpacityMask);
+					}
+				}
+
+				return MaterialNode;
+			}
+
+			void FFbxMaterial::AddAllNodeMaterials(UInterchangeBaseNode* UnrealNode, FbxNode* ParentFbxNode, UInterchangeBaseNodeContainer& NodeContainer, TArray<FString>& JSonErrorMessages)
+			{
+				int32 MaterialCount = ParentFbxNode->GetMaterialCount();
+				for (int32 MaterialIndex = 0; MaterialIndex < MaterialCount; ++MaterialIndex)
+				{
+					FbxSurfaceMaterial* SurfaceMaterial = ParentFbxNode->GetMaterial(MaterialIndex);
+					UInterchangeMaterialNode* MaterialNode = AddNodeMaterial(SurfaceMaterial, NodeContainer, JSonErrorMessages);
+					//The dependencies order is important because mesh will use index in that order to determine material use by a face
+					UnrealNode->SetDependencyUID(MaterialNode->GetUniqueID());
+				}
+			}
+
+			void FFbxMaterial::AddAllSceneMaterials(FbxScene* SDKScene, UInterchangeBaseNodeContainer& NodeContainer, TArray<FString>& JSonErrorMessages)
+			{
+				int32 NodeMaterialCount = SDKScene->GetMaterialCount();
+				for (int32 NodeMaterialIndex = 0; NodeMaterialIndex < NodeMaterialCount; ++NodeMaterialIndex)
+				{
+					FbxSurfaceMaterial* SurfaceMaterial = SDKScene->GetMaterial(NodeMaterialIndex);
+					AddNodeMaterial(SurfaceMaterial, NodeContainer, JSonErrorMessages);
+				}
+			}
+		} //ns Private
+	} //ns FbxParser
+}//ns UE
diff --git a/Engine/Plugins/Runtime/Interchange/Source/Parsers/Fbx/Private/FbxMaterial.h b/Engine/Plugins/Runtime/Interchange/Source/Parsers/Fbx/Private/FbxMaterial.h
new file mode 100644
--- /dev/null
+++ b/Engine/Plugins/Runtime/Interchange/Source/Parsers/Fbx/Private/FbxMaterial.h
@@ -0,0 +1,33 @@
+// Copyright Epic Games, Inc. All Rights Reserved.
+
+#pragma once
+
+#include "CoreMinimal.h"
+#include "FbxInclude.h"
+
+/** Forward declarations */
+class UInterchangeBaseNode;
+class UInterchangeBaseNodeContainer;
+class UInterchangeMaterialNode;
+class UInterchangeTextureNode;
+
+namespace UE
+{
+	namespace FbxParser
+	{
+		namespace Private
+		{
+			class FFbxMaterial
+			{
+			public:
+				static void AddAllNodeMaterials(UInterchangeBaseNode* ParentNode, FbxNode* ParentFbxNode, UInterchangeBaseNodeContainer& NodeContainer, TArray<FString>& JSonErrorMessages);
+				static void AddAllSceneMaterials(FbxScene* SDKScene, UInterchangeBaseNodeContainer& NodeContainer, TArray<FString>& JSonErrorMessages);
+				static UInterchangeMaterialNode* CreateMaterialNode(UInterchangeBaseNodeContainer& NodeContainer, const FString& NodeName, TArray<FString>& JSonErrorMessages);
+				static UInterchangeTextureNode* CreateTexture2DNode(UInterchangeBaseNodeContainer& NodeContainer, const FString& TextureFilePath, TArray<FString>& JSonErrorMessages);
+			
+			private:
+				static UInterchangeMaterialNode* AddNodeMaterial(FbxSurfaceMaterial* SurfaceMaterial, UInterchangeBaseNodeContainer& NodeContainer, TArray<FString>& JSonErrorMessages);
+			};
+		}
+	}
+}
diff --git a/Engine/Plugins/Runtime/Interchange/Source/Parsers/Fbx/Private/FbxScene.cpp b/Engine/Plugins/Runtime/Interchange/Source/Parsers/Fbx/Private/FbxScene.cpp
new file mode 100644
--- /dev/null
+++ b/Engine/Plugins/Runtime/Interchange/Source/Parsers/Fbx/Private/FbxScene.cpp
@@ -0,0 +1,116 @@
+// Copyright Epic Games, Inc. All Rights Reserved.
+
+#include "FbxScene.h"
+
+#include "CoreMinimal.h"
+#include "FbxConvert.h"
+#include "FbxInclude.h"
+#include "FbxMaterial.h"
+#include "InterchangeSceneNode.h"
+#include "Nodes/InterchangeBaseNodeContainer.h"
+
+namespace UE
+{
+	namespace FbxParser
+	{
+		namespace Private
+		{
+			namespace ScenePrivate
+			{
+				UInterchangeSceneNode* AddHierarchyRecursively(FbxNode* Node, FbxScene* SDKScene, UInterchangeBaseNodeContainer& NodeContainer, TArray<FString>& JSonErrorMessages)
+				{
+					FString NodeName = FFbxConvert::MakeString(FFbxConvert::MakeName(Node->GetName()));
+					UInterchangeSceneNode* UnrealNode = FFbxScene::CreateTransformNode(NodeContainer, NodeName, JSonErrorMessages);
+					NodeContainer.AddNode(UnrealNode);
+					FTransform LocalTransform = FFbxConvert::GetTransform(Node->EvaluateLocalTransform());
+					UnrealNode->SetCustomLocalTransform(LocalTransform);
+					FTransform GlobalTransform = FFbxConvert::GetTransform(Node->EvaluateGlobalTransform());
+					UnrealNode->SetCustomGlobalTransform(GlobalTransform);
+
+					int32 AttributeCount = Node->GetNodeAttributeCount();
+					for (int32 AttributeIndex = 0; AttributeIndex < AttributeCount; ++AttributeIndex)
+					{
+						const FbxNodeAttribute* NodeAttribute = Node->GetNodeAttributeByIndex(AttributeIndex);
+						switch (NodeAttribute->GetAttributeType())
+						{
+							case FbxNodeAttribute::eLODGroup:
+								//Add LOD group node
+								//UnrealNode = FFbxScene::CreateLodGroupNode(NodeContainer, NodeName, JSonErrorMessages);
+								break;
+							case FbxNodeAttribute::eMesh:
+								//Add a mesh node
+								if (Node->GetMesh()->GetDeformerCount(FbxDeformer::eSkin) > 0)
+								{
+									//Skeletal mesh node
+									//UnrealNode = FFbxSkeletalMesh::CreateSkeletalMeshNode(NodeContainer, NodeName, JSonErrorMessages);
+								}
+								else
+								{
+									//Static mesh node
+									//UnrealNode = FFbxStaticMesh::CreateStaticMeshNode(NodeContainer, NodeName, JSonErrorMessages);
+								}
+								break;
+							case FbxNodeAttribute::eSkeleton:
+								//Add a joint node
+								//UnrealNode = FFbxSkeletalMesh::CreateJointNode(NodeContainer, NodeName, JSonErrorMessages);
+								break;
+							case FbxNodeAttribute::eCamera:
+								//Add a camera node
+								//UnrealNode = FFbxCamera::CreateCameraNode(NodeContainer, NodeName, JSonErrorMessages);
+								break;
+							case FbxNodeAttribute::eCameraSwitcher:
+								//Add a camera switcher node
+								//UnrealNode = FFbxCamera::CreateCameraSwitcherNode(NodeContainer, NodeName, JSonErrorMessages);
+								break;
+							case FbxNodeAttribute::eLight:
+								//Add a light node
+								//UnrealNode = FFbxLight::CreateLightNode(NodeContainer, NodeName, JSonErrorMessages);
+								break;
+							default:
+								//Add a transform node
+								UnrealNode = FFbxScene::CreateTransformNode(NodeContainer, NodeName, JSonErrorMessages);
+								break;
+						}
+					}
+					if (UnrealNode)
+					{
+						NodeContainer.AddNode(UnrealNode);
+						//Add the dependencies of the material on the correct order
+						FFbxMaterial::AddAllNodeMaterials(UnrealNode, Node, NodeContainer, JSonErrorMessages);
+					}
+
+					const int32 ChildCount = Node->GetChildCount();
+					for (int32 ChildIndex = 0; ChildIndex < ChildCount; ++ChildIndex)
+					{
+						FbxNode* ChildNode = Node->GetChild(ChildIndex);
+						UInterchangeSceneNode* UnrealChildNode = AddHierarchyRecursively(ChildNode, SDKScene, NodeContainer, JSonErrorMessages);
+						UnrealChildNode->SetParentUID(UnrealNode->GetUniqueID());
+					}
+					return UnrealNode;
+				}
+			} //ns ScenePrivate
+
+			UInterchangeSceneNode* FFbxScene::CreateTransformNode(UInterchangeBaseNodeContainer& NodeContainer, const FString& NodeName, TArray<FString>& JSonErrorMessages)
+			{
+				FName DisplayLabel = *NodeName;
+				FName NodeUID(*NodeName);
+				UInterchangeSceneNode* TransformNode = NewObject<UInterchangeSceneNode>(&NodeContainer, NAME_None);
+				if (!ensure(TransformNode))
+				{
+					JSonErrorMessages.Add(TEXT("{\"Msg\" : {\"Type\" : \"Error\",\n\"Msg\" : \"Cannot allocate a node when importing fbx\"}}"));
+					return nullptr;
+				}
+				// Creating a UMaterialInterface
+				TransformNode->InitializeNode(NodeUID, DisplayLabel);
+				NodeContainer.AddNode(TransformNode);
+				return TransformNode;
+			}
+
+			void FFbxScene::AddHierarchy(FbxScene* SDKScene, UInterchangeBaseNodeContainer& NodeContainer, TArray<FString>& JSonErrorMessages)
+			{
+				 FbxNode* RootNode = SDKScene->GetRootNode();
+				 ScenePrivate::AddHierarchyRecursively(RootNode, SDKScene, NodeContainer, JSonErrorMessages);
+			}
+		} //ns Private
+	} //ns FbxParser
+}//ns UE
diff --git a/Engine/Plugins/Runtime/Interchange/Source/Parsers/Fbx/Private/FbxScene.h b/Engine/Plugins/Runtime/Interchange/Source/Parsers/Fbx/Private/FbxScene.h
new file mode 100644
--- /dev/null
+++ b/Engine/Plugins/Runtime/Interchange/Source/Parsers/Fbx/Private/FbxScene.h
@@ -0,0 +1,26 @@
+// Copyright Epic Games, Inc. All Rights Reserved.
+
+#pragma once
+
+#include "CoreMinimal.h"
+#include "FbxInclude.h"
+
+/** Forward declarations */
+class UInterchangeBaseNodeContainer;
+class UInterchangeSceneNode;
+
+namespace UE
+{
+	namespace FbxParser
+	{
+		namespace Private
+		{
+			class FFbxScene
+			{
+			public:
+				static void AddHierarchy(FbxScene* SDKScene, UInterchangeBaseNodeContainer& NodeContainer, TArray<FString>& JSonErrorMessages);
+				static UInterchangeSceneNode* CreateTransformNode(UInterchangeBaseNodeContainer& NodeContainer, const FString& NodeName, TArray<FString>& JSonErrorMessages);
+			};
+		}
+	}
+}
diff --git a/Engine/Plugins/Runtime/Interchange/Source/Parsers/Fbx/Private/InterchangeFbxParser.cpp b/Engine/Plugins/Runtime/Interchange/Source/Parsers/Fbx/Private/InterchangeFbxParser.cpp
--- a/Engine/Plugins/Runtime/Interchange/Source/Parsers/Fbx/Private/InterchangeFbxParser.cpp
+++ b/Engine/Plugins/Runtime/Interchange/Source/Parsers/Fbx/Private/InterchangeFbxParser.cpp
@@ -3,14 +3,36 @@
 #include "InterchangeFbxParser.h"
 
 #include "CoreMinimal.h"
+#include "FbxAPI.h"
+#include "Nodes/InterchangeBaseNodeContainer.h"
+#include "InterchangeTextureNode.h"
+#include "UObject/StrongObjectPtr.h"
+#include "UObject/UObjectGlobals.h"
+#include "UObject/Package.h"
 
 namespace InterchangeFbxParser
 {
 	void FbxParser::LoadFbxFile(const FString& Filename, const FString& ResultFolder)
 	{
-		ResultFilepath = ResultFolder + TEXT("/Foo.txt");
-		JsonLoadMessages = TEXT("{\"Msg\" : {\"Type\" : \"Log\",\n\"Msg\" : \"This is a success!\"}}");
-		return;
-	}
+		UE::FbxParser::Private::FFbxParser FbxParserPrivate;
+		if (!FbxParserPrivate.LoadFbxFile(Filename, JsonLoadMessages))
+		{
+			JsonLoadMessages.Add(TEXT("{\"Msg\" : {\"Type\" : \"Error\",\n\"Msg\" : \"Cannot load the fbx file.\"}}"));
+			return;
+		}
 
+		
+		ResultFilepath = ResultFolder + TEXT("/SceneDescription.itc");
+		
+		Container = TStrongObjectPtr<UInterchangeBaseNodeContainer>(NewObject<UInterchangeBaseNodeContainer>(GetTransientPackage(), NAME_None));
+		if (!ensure(Container.IsValid()))
+		{
+			JsonLoadMessages.Add(TEXT("{\"Msg\" : {\"Type\" : \"Error\",\n\"Msg\" : \"Cannot allocate base node container to add fbx scene data\"}}"));
+			return;
+		}
+		FbxParserPrivate.FillContainerWithFbxScene(*Container.Get(), JsonLoadMessages);
+		
+		Container.Get()->SaveToFile(ResultFilepath);
+		JsonLoadMessages.Add(TEXT("{\"Msg\" : {\"Type\" : \"Log\",\n\"Msg\" : \"This is a success!\"}}"));
+	}
 }
diff --git a/Engine/Plugins/Runtime/Interchange/Source/Parsers/Fbx/Public/InterchangeFbxParser.h b/Engine/Plugins/Runtime/Interchange/Source/Parsers/Fbx/Public/InterchangeFbxParser.h
--- a/Engine/Plugins/Runtime/Interchange/Source/Parsers/Fbx/Public/InterchangeFbxParser.h
+++ b/Engine/Plugins/Runtime/Interchange/Source/Parsers/Fbx/Public/InterchangeFbxParser.h
@@ -3,7 +3,8 @@
 #pragma once
 
 #include "CoreMinimal.h"
-#include "Nodes/BaseNodeContainer.h"
+#include "Nodes/InterchangeBaseNodeContainer.h"
+#include "UObject/StrongObjectPtr.h"
 
 namespace InterchangeFbxParser
 {
@@ -19,12 +20,12 @@ public:
 	void LoadFbxFile(const FString& Filename, const FString& ResultFolder);
 	
 	FString GetResultFilepath(){ return ResultFilepath; }
-	FString GetJsonLoadMessages() { return JsonLoadMessages; }
+	TArray<FString> GetJsonLoadMessages() { return JsonLoadMessages; }
 
 private:
-	Interchange::FBaseNodeContainer Container;
+	TStrongObjectPtr<UInterchangeBaseNodeContainer> Container = nullptr;
 	FString ResultFilepath;
-	FString JsonLoadMessages;
+	TArray<FString> JsonLoadMessages;
 };
 
 } // ns InterchangeFbxParser
diff --git a/Engine/Source/Programs/InterchangeWorker/InterchangeWorker.Build.cs b/Engine/Source/Programs/InterchangeWorker/InterchangeWorker.Build.cs
--- a/Engine/Source/Programs/InterchangeWorker/InterchangeWorker.Build.cs
+++ b/Engine/Source/Programs/InterchangeWorker/InterchangeWorker.Build.cs
@@ -16,6 +16,7 @@ public class InterchangeWorker : ModuleRules
 			{
 				"ApplicationCore",
 				"Core",
+				"CoreUObject",
 				"InterchangeDispatcher",
 				"Json",
 				"Projects",
@@ -23,11 +24,18 @@ public class InterchangeWorker : ModuleRules
 			}
 		);
 
-		if (Target.Platform == UnrealTargetPlatform.Win64)
+		if (Target.Platform == UnrealTargetPlatform.Win64 ||
+			Target.Platform == UnrealTargetPlatform.Linux ||
+			Target.Platform == UnrealTargetPlatform.Mac)
 		{
-			PrivateDependencyModuleNames.Add("InterchangeFbxParser");
+			PrivateDependencyModuleNames.AddRange(
+				new string[]
+				{
+					"InterchangeCore",
+					"InterchangeFbxParser",
+					"InterchangeNodePlugin"
+				}
+			);
 		}
-
-		PublicDelayLoadDLLs.Add("kernel_io.dll");
 	}
 }
diff --git a/Engine/Source/Programs/InterchangeWorker/InterchangeWorker.Target.cs b/Engine/Source/Programs/InterchangeWorker/InterchangeWorker.Target.cs
--- a/Engine/Source/Programs/InterchangeWorker/InterchangeWorker.Target.cs
+++ b/Engine/Source/Programs/InterchangeWorker/InterchangeWorker.Target.cs
@@ -39,6 +39,8 @@ public class InterchangeWorkerTarget : TargetRules
 		bCompileICU = false;
 		bCompileCEF3 = false;
 
+		GlobalDefinitions.Add("USE_LOCALIZED_PACKAGE_CACHE=0");
+
 		string EngineDir = @"..\.."; // relative to default destination, which is $(EngineDir)\Binaries\Win64
 		ExeBinariesSubFolder = EngineDir + @"\Plugins\Runtime\Interchange\Binaries\Win64\";
 	}
diff --git a/Engine/Source/Programs/InterchangeWorker/Private/InterchangeWorker.cpp b/Engine/Source/Programs/InterchangeWorker/Private/InterchangeWorker.cpp
--- a/Engine/Source/Programs/InterchangeWorker/Private/InterchangeWorker.cpp
+++ b/Engine/Source/Programs/InterchangeWorker/Private/InterchangeWorker.cpp
@@ -1,8 +1,9 @@
 // Copyright Epic Games, Inc. All Rights Reserved.
 
 #include "InterchangeWorker.h"
-#include "InterchangeWorkerImpl.h"
 
+#include "HAL/PlatformProcess.h"
+#include "InterchangeWorkerImpl.h"
 #include "RequiredProgramMainCPPInclude.h"
 
 IMPLEMENT_APPLICATION(InterchangeWorker, "InterchangeWorker");
@@ -69,24 +70,36 @@ int32 Main(int32 Argc, TCHAR * Argv[])
 	return EXIT_SUCCESS;
 }
 
+#if PLATFORM_WINDOWS
 int32 Filter(uint32 Code, struct _EXCEPTION_POINTERS *Ep)
 {
 	return EXCEPTION_EXECUTE_HANDLER;
 }
+#endif //#if PLATFORM_WINDOWS
 
 INT32_MAIN_INT32_ARGC_TCHAR_ARGV()
 {
+	int32 ReturnCode = 0;
 	GEngineLoop.PreInit(ArgC, ArgV);
 
+#if PLATFORM_WINDOWS
 	SetErrorMode(SEM_FAILCRITICALERRORS | SEM_NOGPFAULTERRORBOX);
 	_set_abort_behavior(0, _WRITE_ABORT_MSG);
 	__try
 	{
-		return Main(ArgC, ArgV);
+		ReturnCode = Main(ArgC, ArgV);
 	}
 	__except (Filter(GetExceptionCode(), GetExceptionInformation()))
 	{
-		return EXIT_FAILURE;
+		ReturnCode = EXIT_FAILURE;
 	}
-	return EXIT_SUCCESS;
+#else // PLATFORM_WINDOWS
+	ReturnCode = Main(ArgC, ArgV);
+#endif // PLATFORM_WINDOWS
+
+	FEngineLoop::AppPreExit();
+	FModuleManager::Get().UnloadModulesAtShutdown();
+	FEngineLoop::AppExit();
+
+	return ReturnCode;
 }
diff --git a/Engine/Source/Programs/InterchangeWorker/Private/InterchangeWorkerImpl.cpp b/Engine/Source/Programs/InterchangeWorker/Private/InterchangeWorkerImpl.cpp
--- a/Engine/Source/Programs/InterchangeWorker/Private/InterchangeWorkerImpl.cpp
+++ b/Engine/Source/Programs/InterchangeWorker/Private/InterchangeWorkerImpl.cpp
@@ -11,6 +11,7 @@
 #endif
 
 #include "HAL/FileManager.h"
+#include "HAL/PlatformFileManager.h"
 #include "Misc/FileHelper.h"
 #include "Misc/Paths.h"
 #include "SocketSubsystem.h"
@@ -125,7 +126,7 @@ void FInterchangeWorkerImpl::ProcessCommand(const FRunTaskCommand& RunTaskComman
 	ETaskState ProcessResult = ETaskState::Unknown;
 	//Process the json and run the task
 	FString JSonResult;
-	FString JSonMessages;
+	TArray<FString> JSonMessages;
 	InterchangeDispatcher::FJsonLoadSourceCmd LoadSourceCommand;
 	//Any command FromJson function return true if the Json descibe the command
 	if (LoadSourceCommand.FromJson(JsonToProcess))
@@ -156,7 +157,7 @@ void FInterchangeWorkerImpl::ProcessCommand(const FRunTaskCommand& RunTaskComman
 	UE_LOG(LogInterchangeWorker, Verbose, TEXT("End of Process %s"), *JsonToProcess);
 }
 
-ETaskState FInterchangeWorkerImpl::LoadFbxFile(const InterchangeDispatcher::FJsonLoadSourceCmd& LoadSourceCommand, FString& OutJSonResult, FString& OutJSonMessages) const
+ETaskState FInterchangeWorkerImpl::LoadFbxFile(const InterchangeDispatcher::FJsonLoadSourceCmd& LoadSourceCommand, FString& OutJSonResult, TArray<FString>& OutJSonMessages) const
 {
 	ETaskState ResultState = ETaskState::Unknown;
 	FString SourceFilename = LoadSourceCommand.GetSourceFilename();
@@ -171,7 +172,7 @@ ETaskState FInterchangeWorkerImpl::LoadFbxFile(const InterchangeDispatcher::FJso
 #else
 	InterchangeDispatcher::FJsonLoadSourceCmd::JsonResultParser ResultParser;
 	ResultParser.SetResultFilename(FString());
-	OutJSonMessages = TEXT("{\"Msg\" : {\"Type\" : \"Error\",\n\"Msg\" : \"Cannot Execute fbx command on other platform then window 64 bits!\"}}");
+	OutJSonMessages.Add(TEXT("{\"Msg\" : {\"Type\" : \"Error\",\n\"Msg\" : \"Cannot Execute fbx command on other platform then window 64 bits!\"}}"));
 	ResultState = ETaskState::ProcessFailed;
 #endif
 	return ResultState;
diff --git a/Engine/Source/Programs/InterchangeWorker/Private/InterchangeWorkerImpl.h b/Engine/Source/Programs/InterchangeWorker/Private/InterchangeWorkerImpl.h
--- a/Engine/Source/Programs/InterchangeWorker/Private/InterchangeWorkerImpl.h
+++ b/Engine/Source/Programs/InterchangeWorker/Private/InterchangeWorkerImpl.h
@@ -24,7 +24,7 @@ private:
 	void ProcessCommand(const InterchangeDispatcher::FBackPingCommand& BackPingCommand);
 	void ProcessCommand(const InterchangeDispatcher::FRunTaskCommand& TerminateCommand);
 
-	InterchangeDispatcher::ETaskState LoadFbxFile(const InterchangeDispatcher::FJsonLoadSourceCmd& LoadSourceCommand, FString& OutJSonResult, FString& OutJSonMessages) const;
+	InterchangeDispatcher::ETaskState LoadFbxFile(const InterchangeDispatcher::FJsonLoadSourceCmd& LoadSourceCommand, FString& OutJSonResult, TArray<FString>& OutJSonMessages) const;
 
 private:
 	InterchangeDispatcher::FNetworkClientNode NetworkInterface;
diff --git a/Engine/Source/Runtime/Interchange/Core/Private/InterchangeManager.cpp b/Engine/Source/Runtime/Interchange/Core/Private/InterchangeManager.cpp
--- a/Engine/Source/Runtime/Interchange/Core/Private/InterchangeManager.cpp
+++ b/Engine/Source/Runtime/Interchange/Core/Private/InterchangeManager.cpp
@@ -13,7 +13,7 @@
 #include "Misc/AsyncTaskNotification.h"
 #include "Misc/MessageDialog.h"
 #include "Misc/ScopedSlowTask.h"
-#include "Nodes/BaseNodeContainer.h"
+#include "Nodes/InterchangeBaseNodeContainer.h"
 #include "PackageUtils/PackageUtils.h"
 #include "Tasks/InterchangeTaskParsing.h"
 #include "Tasks/InterchangeTaskPipeline.h"
@@ -93,8 +93,8 @@ void Interchange::FImportAsyncHelper::AddReferencedObjects(FReferenceCollector&
 
 void Interchange::FImportAsyncHelper::CleanUp()
 {
-	//Release the node graph adapters
-	BaseNodeContainerAdapters.Empty();
+	//Release the graph
+	BaseNodeContainers.Empty();
 
 	for (UInterchangeSourceData* SourceData : SourceDatas)
 	{
@@ -295,13 +295,11 @@ Interchange::FAsyncImportResult UInterchangeManager::ImportAssetAsync(const FStr
 		ensure(AsyncHelper->Translators.Add(GetTranslatorForSourceData(AsyncHelper->SourceDatas[SourceDataIndex])) == SourceDataIndex);
 	}
 
-	//Create the node graphs and the adapters (TStrongObjectPtr must be use only on the main thread, so we create them here before filling the node Containers)
-	AsyncHelper->BaseNodeContainers.AddDefaulted(AsyncHelper->SourceDatas.Num());
+	//Create the node graphs for each source data (StrongObjectPtr has to be created on the main thread)
 	for (int32 SourceDataIndex = 0; SourceDataIndex < AsyncHelper->SourceDatas.Num(); ++SourceDataIndex)
 	{
-		AsyncHelper->BaseNodeContainerAdapters.Add(TStrongObjectPtr<UInterchangeBaseNodeContainerAdapter>(NewObject<UInterchangeBaseNodeContainerAdapter>(GetTransientPackage(), NAME_None)));
-		check(AsyncHelper->BaseNodeContainerAdapters[SourceDataIndex].IsValid());
-		AsyncHelper->BaseNodeContainerAdapters[SourceDataIndex].Get()->SetBaseNodeContainer(&AsyncHelper->BaseNodeContainers[SourceDataIndex]);
+		AsyncHelper->BaseNodeContainers.Add(TStrongObjectPtr<UInterchangeBaseNodeContainer>(NewObject<UInterchangeBaseNodeContainer>(GetTransientPackage(), NAME_None)));
+		check(AsyncHelper->BaseNodeContainers[SourceDataIndex].IsValid());
 	}
 
 	if ( ImportAssetParameters.OverridePipeline == nullptr )
diff --git a/Engine/Source/Runtime/Interchange/Core/Private/Nodes/InterchangeBaseNode.cpp b/Engine/Source/Runtime/Interchange/Core/Private/Nodes/InterchangeBaseNode.cpp
new file mode 100644
--- /dev/null
+++ b/Engine/Source/Runtime/Interchange/Core/Private/Nodes/InterchangeBaseNode.cpp
@@ -0,0 +1,204 @@
+// Copyright Epic Games, Inc. All Rights Reserved.
+#include "Nodes/InterchangeBaseNode.h"
+
+#include "CoreMinimal.h"
+#include "Types/AttributeStorage.h"
+#include "UObject/Class.h"
+#include "UObject/Object.h"
+#include "UObject/ObjectMacros.h"
+
+
+
+
+void UInterchangeBaseNode::InitializeNode(const FName& UniqueID, const FName& DisplayLabel)
+{
+	Interchange::EAttributeStorageResult Result = Attributes.RegisterAttribute(Interchange::FBaseNodeStaticData::UniqueIDKey(), UniqueID, Interchange::EAttributeProperty::NoHash);
+	if (!IsAttributeStorageResultSuccess(Result))
+	{
+		LogAttributeStorageErrors(Result, TEXT("RegisterAttribute"), Interchange::FBaseNodeStaticData::UniqueIDKey());
+	}
+
+	Result = Attributes.RegisterAttribute(Interchange::FBaseNodeStaticData::DisplayLabelKey(), DisplayLabel, Interchange::EAttributeProperty::NoHash);
+	if (!IsAttributeStorageResultSuccess(Result))
+	{
+		LogAttributeStorageErrors(Result, TEXT("RegisterAttribute"), Interchange::FBaseNodeStaticData::DisplayLabelKey());
+	}
+	bIsInitialized = true;
+}
+
+FString UInterchangeBaseNode::GetTypeName() const
+{
+	const FString TypeName = TEXT("BaseNode");
+	return TypeName;
+}
+
+bool UInterchangeBaseNode::HasAttribute(const Interchange::FAttributeKey& NodeAttributeKey) const
+{
+	return Attributes.ContainAttribute(NodeAttributeKey);
+}
+
+
+FName UInterchangeBaseNode::GetUniqueID() const
+{
+	ensure(bIsInitialized);
+	FName UniqueID = NAME_None;
+	Attributes.GetAttributeHandle<FName>(Interchange::FBaseNodeStaticData::UniqueIDKey()).Get(UniqueID);
+	return UniqueID;
+}
+
+FName UInterchangeBaseNode::GetDisplayLabel() const
+{
+	ensure(bIsInitialized);
+	checkSlow(Attributes.ContainAttribute(Interchange::FBaseNodeStaticData::DisplayLabelKey()));
+	FName DisplayLabel = NAME_None;
+	Attributes.GetAttributeHandle<FName>(Interchange::FBaseNodeStaticData::DisplayLabelKey()).Get(DisplayLabel);
+	return DisplayLabel;
+}
+
+FName UInterchangeBaseNode::GetParentUID() const
+{
+	if (!Attributes.ContainAttribute(Interchange::FBaseNodeStaticData::ParentIDKey()))
+	{
+		return InvalidNodeUID();
+	}
+
+	FName ParentUniqueID = NAME_None;
+	Interchange::FAttributeStorage::TAttributeHandle<FName> Handle = Attributes.GetAttributeHandle<FName>(Interchange::FBaseNodeStaticData::ParentIDKey());
+	if(Handle.IsValid())
+	{
+		Handle.Get(ParentUniqueID);
+		return ParentUniqueID;
+	}
+	return InvalidNodeUID();
+}
+
+bool UInterchangeBaseNode::SetParentUID(FName ParentUID)
+{
+	Interchange::EAttributeStorageResult Result = Attributes.RegisterAttribute(Interchange::FBaseNodeStaticData::ParentIDKey(), ParentUID);
+	if(IsAttributeStorageResultSuccess(Result))
+	{
+		Interchange::FAttributeStorage::TAttributeHandle<FName> Handle = Attributes.GetAttributeHandle<FName>(Interchange::FBaseNodeStaticData::ParentIDKey());
+		return Handle.IsValid();
+	}
+	return false;
+}
+
+void UInterchangeBaseNode::GetDependecies(TArray<FName>& OutDependencies ) const
+{
+	OutDependencies.Reset();
+
+	int32 DepenedenciesCount = 0;
+	if (!Attributes.ContainAttribute(Interchange::FBaseNodeStaticData::DependencyCountKey()))
+	{
+		return;
+	}
+	Interchange::FAttributeStorage::TAttributeHandle<int32> Handle = Attributes.GetAttributeHandle<int32>(Interchange::FBaseNodeStaticData::DependencyCountKey());
+	if (!Handle.IsValid())
+	{
+		return;
+	}
+	Handle.Get(DepenedenciesCount);
+	for (int32 DepIndex = 0; DepIndex < DepenedenciesCount; ++DepIndex)
+	{
+		FString DepIndexKeyString = Interchange::FBaseNodeStaticData::DependencyBaseKey().Key.ToString() + FString::FromInt(DepIndex);
+		Interchange::FAttributeKey DepIndexKey(*DepIndexKeyString);
+		Interchange::FAttributeStorage::TAttributeHandle<FName> HandleDep = Attributes.GetAttributeHandle<FName>(DepIndexKey);
+		if (!HandleDep.IsValid())
+		{
+			continue;
+		}
+		FName& NodeUniqueID = OutDependencies.AddDefaulted_GetRef();
+		HandleDep.Get(NodeUniqueID);
+	}
+}
+
+bool UInterchangeBaseNode::SetDependencyUID(FName DependencyUID)
+{
+	if (!Attributes.ContainAttribute(Interchange::FBaseNodeStaticData::DependencyCountKey()))
+	{
+		const int32 DependencyCount = 0;
+		Interchange::EAttributeStorageResult Result = Attributes.RegisterAttribute<int32>(Interchange::FBaseNodeStaticData::DependencyCountKey(), DependencyCount);
+	}
+	Interchange::FAttributeStorage::TAttributeHandle<int32> Handle = Attributes.GetAttributeHandle<int32>(Interchange::FBaseNodeStaticData::DependencyCountKey());
+	if(!ensure(Handle.IsValid()))
+	{
+		return false;
+	}
+	int32 DepIndex = 0;
+	Handle.Get(DepIndex);
+	FString DepIndexKeyString = Interchange::FBaseNodeStaticData::DependencyBaseKey().Key.ToString() + FString::FromInt(DepIndex);
+	Interchange::FAttributeKey DepIndexKey(*DepIndexKeyString);
+	//Increment the counter
+	DepIndex++;
+	Handle.Set(DepIndex);
+
+	Interchange::EAttributeStorageResult DepResult = Attributes.RegisterAttribute<FName>(DepIndexKey, DependencyUID);
+	return true;
+}
+
+bool UInterchangeBaseNode::IsEnabled() const
+{
+	Interchange::FAttributeStorage::TAttributeHandle<bool> Handle = Attributes.GetAttributeHandle<bool>(Interchange::FBaseNodeStaticData::IsEnabledKey());
+	if (Handle.IsValid())
+	{
+		bool bIsEnabled = false;
+		Handle.Get(bIsEnabled);
+		return bIsEnabled;
+	}
+	return true;
+}
+
+bool UInterchangeBaseNode::SetEnabled(const bool bIsEnabled)
+{
+	Interchange::EAttributeStorageResult Result = Attributes.RegisterAttribute(Interchange::FBaseNodeStaticData::IsEnabledKey(), bIsEnabled);
+	if (IsAttributeStorageResultSuccess(Result))
+	{
+		Interchange::FAttributeStorage::TAttributeHandle<bool> Handle = Attributes.GetAttributeHandle<bool>(Interchange::FBaseNodeStaticData::IsEnabledKey());
+		return Handle.IsValid();
+	}
+	return false;
+}
+
+FGuid UInterchangeBaseNode::GetHash() const
+{
+	return Attributes.GetStorageHash();
+}
+
+class UClass* UInterchangeBaseNode::GetAssetClass() const
+{
+	return nullptr;
+}
+
+FName UInterchangeBaseNode::InvalidNodeUID()
+{
+	return NAME_None;
+}
+
+void UInterchangeBaseNode::ApplyAllCustomAttributeToAsset(UObject* Object) const
+{
+	UClass* ObjectClass = Object->GetClass();
+	for (const TPair<UClass*, TArray<Interchange::FApplyAttributeToAsset>>& ClassDelegatePair : ApplyCustomAttributeDelegates)
+	{
+		if(ObjectClass->IsChildOf(ClassDelegatePair.Key))
+		{
+			for (const Interchange::FApplyAttributeToAsset& Delegate : ClassDelegatePair.Value)
+			{
+				if(Delegate.IsBound())
+				{
+					Delegate.Execute(Object);
+				}
+			}
+		}
+	}
+}
+
+void UInterchangeBaseNode::Serialize(FArchive& Ar)
+{
+	Ar << Attributes;
+	if (Ar.IsLoading())
+	{
+		//The node is consider Initialize if the UniqueID and the Display label are set properly
+		bIsInitialized = (Attributes.GetAttributeHandle<FName>(Interchange::FBaseNodeStaticData::UniqueIDKey()).IsValid() &&
+							Attributes.GetAttributeHandle<FName>(Interchange::FBaseNodeStaticData::DisplayLabelKey()).IsValid());
+	}
+}
diff --git a/Engine/Source/Runtime/Interchange/Core/Private/Nodes/InterchangeBaseNodeContainer.cpp b/Engine/Source/Runtime/Interchange/Core/Private/Nodes/InterchangeBaseNodeContainer.cpp
new file mode 100644
--- /dev/null
+++ b/Engine/Source/Runtime/Interchange/Core/Private/Nodes/InterchangeBaseNodeContainer.cpp
@@ -0,0 +1,228 @@
+// Copyright Epic Games, Inc. All Rights Reserved.
+#include "Nodes/InterchangeBaseNodeContainer.h"
+
+#include "CoreMinimal.h"
+#include "Misc/FileHelper.h"
+#include "Nodes/InterchangeBaseNode.h"
+#include "Serialization/LargeMemoryWriter.h"
+#include "Serialization/LargeMemoryReader.h"
+#include "Templates/UniquePtr.h"
+#include "UObject/Class.h"
+#include "UObject/Object.h"
+#include "UObject/ObjectMacros.h"
+#include "UObject/UObjectIterator.h"
+
+FName UInterchangeBaseNodeContainer::AddNode(UInterchangeBaseNode* Node)
+{
+	if (!Node)
+	{
+		return UInterchangeBaseNode::InvalidNodeUID();
+	}
+	FName NodeUniqueID = Node->GetUniqueID();
+	if (NodeUniqueID == UInterchangeBaseNode::InvalidNodeUID())
+	{
+		return UInterchangeBaseNode::InvalidNodeUID();
+	}
+		
+	//Cannot add an node with the same IDs
+	if (Nodes.Contains(NodeUniqueID))
+	{
+		return NodeUniqueID;
+	}
+
+	//Copy the node
+	Nodes.Add(NodeUniqueID, Node);
+	return NodeUniqueID;
+}
+
+bool UInterchangeBaseNodeContainer::IsNodeUIDValid(FName NodeUniqueID) const
+{
+	if (NodeUniqueID == UInterchangeBaseNode::InvalidNodeUID())
+	{
+		return false;
+	}
+	return Nodes.Contains(NodeUniqueID);
+}
+
+void UInterchangeBaseNodeContainer::IterateNodes(TFunctionRef<void(FName, const UInterchangeBaseNode*)> IterationLambda)
+{
+	for (TPair<FName, UInterchangeBaseNode*>& NodeKeyValue : Nodes)
+	{
+		IterationLambda(NodeKeyValue.Key, NodeKeyValue.Value);
+	}
+}
+
+void UInterchangeBaseNodeContainer::GetRoots(TArray<FName>& RootNodes)
+{
+	for (TPair<FName, UInterchangeBaseNode*>& NodeKeyValue : Nodes)
+	{
+		if (NodeKeyValue.Value->GetParentUID() == UInterchangeBaseNode::InvalidNodeUID())
+		{
+			RootNodes.Add(NodeKeyValue.Key);
+		}
+	}
+}
+
+UInterchangeBaseNode* UInterchangeBaseNodeContainer::GetNode(FName NodeUniqueID)
+{
+	if (NodeUniqueID == UInterchangeBaseNode::InvalidNodeUID())
+	{
+		return nullptr;
+	}
+	if (!Nodes.Contains(NodeUniqueID))
+	{
+		return nullptr;
+	}
+	UInterchangeBaseNode* Node = Nodes.FindChecked(NodeUniqueID);
+	return Node;
+}
+
+const UInterchangeBaseNode* UInterchangeBaseNodeContainer::GetNode(FName NodeUniqueID) const
+{
+	return const_cast<UInterchangeBaseNodeContainer*>(this)->GetNode(NodeUniqueID);
+}
+
+bool UInterchangeBaseNodeContainer::SetNodeParentUID(FName NodeUniqueID, FName NewParentNodeUID)
+{
+	if (!Nodes.Contains(NodeUniqueID))
+	{
+		return false;
+	}
+	if (!Nodes.Contains(NewParentNodeUID))
+	{
+		return false;
+	}
+	UInterchangeBaseNode* Node = Nodes.FindChecked(NodeUniqueID);
+	Node->SetParentUID(NewParentNodeUID);
+	return true;
+}
+
+int32 UInterchangeBaseNodeContainer::GetNodeChildrenCount(FName NodeUniqueID) const
+{
+	TArray<FName> ChildrenUIDs = GetNodeChildrenUIDs(NodeUniqueID);
+	return ChildrenUIDs.Num();
+}
+
+TArray<FName> UInterchangeBaseNodeContainer::GetNodeChildrenUIDs(FName NodeUniqueID) const
+{
+	TArray<FName> OutChildrenUIDs;
+	for (const TPair<FName, UInterchangeBaseNode*>& NodeKeyValue : Nodes)
+	{
+		if (NodeKeyValue.Value->GetParentUID() == NodeUniqueID)
+		{
+			OutChildrenUIDs.Add(NodeKeyValue.Key);
+		}
+	}
+	return OutChildrenUIDs;
+}
+
+UInterchangeBaseNode* UInterchangeBaseNodeContainer::GetNodeChildren(FName NodeUniqueID, int32 ChildIndex)
+{
+	return GetNodeChildrenInternal(NodeUniqueID, ChildIndex);
+}
+
+const UInterchangeBaseNode* UInterchangeBaseNodeContainer::GetNodeChildren(FName NodeUniqueID, int32 ChildIndex) const
+{
+	return const_cast<UInterchangeBaseNodeContainer*>(this)->GetNodeChildrenInternal(NodeUniqueID, ChildIndex);
+}
+
+void UInterchangeBaseNodeContainer::SerializeNodeContainerData(FArchive& Ar)
+{
+	if (Ar.IsLoading())
+	{
+		Nodes.Reset();
+	}
+	int32 NodeCount = Nodes.Num();
+	Ar << NodeCount;
+
+	if(Ar.IsSaving())
+	{
+		//The node name is not serialize since its an attribute inside the node that will be serialize by the node himself
+		auto SerializeNodePair = [&Ar](UInterchangeBaseNode* BaseNode)
+		{
+			FString ClassFullName = BaseNode->GetClass()->GetFullName();
+			Ar << ClassFullName;
+			BaseNode->Serialize(Ar);
+		};
+
+		for (TPair<FName, UInterchangeBaseNode*> NodePair : Nodes)
+		{
+			SerializeNodePair(NodePair.Value);
+		}
+	}
+	else if(Ar.IsLoading())
+	{
+		//Find all the potential node class
+		TMap<FString, UClass*> ClassPerName;
+		for (FObjectIterator It(UClass::StaticClass()); It; ++It)
+		{
+			UClass* Class = Cast<UClass>(*It);
+			if (Class->IsChildOf(UInterchangeBaseNode::StaticClass()))
+			{
+				ClassPerName.Add(Class->GetFullName(), Class);
+			}
+		}
+
+		for (int32 NodeIndex = 0; NodeIndex < NodeCount; ++NodeIndex)
+		{
+			FString ClassFullName;
+			Ar << ClassFullName;
+			//This cannot fail to make sure we have a healty serialization
+			if (!ensure(ClassPerName.Contains(ClassFullName)))
+			{
+				//We did not successfully serialize the content of the file into the node container
+				return;
+			}
+			UClass* ToCreateClass = ClassPerName.FindChecked(ClassFullName);
+			//Create a UInterchangeBaseNode with the proper class
+			UInterchangeBaseNode* BaseNode = NewObject<UInterchangeBaseNode>(this, ToCreateClass);
+			BaseNode->Serialize(Ar);
+			AddNode(BaseNode);
+		}
+	}
+}
+
+void UInterchangeBaseNodeContainer::SaveToFile(const FString& Filename)
+{
+	FLargeMemoryWriter Ar;
+	SerializeNodeContainerData(Ar);
+	uint8* ArchiveData = Ar.GetData();
+	int64 ArchiveSize = Ar.TotalSize();
+	TArray64<uint8> Buffer(ArchiveData, ArchiveSize);
+	FFileHelper::SaveArrayToFile(Buffer, *Filename);
+}
+
+void UInterchangeBaseNodeContainer::LoadFromFile(const FString& Filename)
+{
+	//All sub object should be gone with the reset
+	Nodes.Reset();
+	TArray64<uint8> Buffer;
+	FFileHelper::LoadFileToArray(Buffer, *Filename);
+	uint8* FileData = Buffer.GetData();
+	int64 FileDataSize = Buffer.Num();
+	if (FileDataSize < 1)
+	{
+		//Nothing to load from this file
+		return;
+	}
+	//Buffer keep the ownership of the data, the large memory reader is use to serialize the TMap
+	FLargeMemoryReader Ar(FileData, FileDataSize);
+	SerializeNodeContainerData(Ar);
+}
+
+UInterchangeBaseNode* UInterchangeBaseNodeContainer::GetNodeChildrenInternal(FName NodeUniqueID, int32 ChildIndex)
+{
+	TArray<FName> ChildrenUIDs = GetNodeChildrenUIDs(NodeUniqueID);
+	if (!ChildrenUIDs.IsValidIndex(ChildIndex))
+	{
+		return nullptr;
+	}
+
+	if (Nodes.Contains(ChildrenUIDs[ChildIndex]))
+	{
+		UInterchangeBaseNode* Node = Nodes.FindChecked(ChildrenUIDs[ChildIndex]);
+		return Node;
+	}
+
+	return nullptr;
+}
diff --git a/Engine/Source/Runtime/Interchange/Core/Private/Nodes/InterchangeBaseNodeContainerAdapter.cpp b/Engine/Source/Runtime/Interchange/Core/Private/Nodes/InterchangeBaseNodeContainerAdapter.cpp
deleted file mode 100644
diff --git a/Engine/Source/Runtime/Interchange/Core/Private/Tasks/InterchangeTaskCompletion.cpp b/Engine/Source/Runtime/Interchange/Core/Private/Tasks/InterchangeTaskCompletion.cpp
--- a/Engine/Source/Runtime/Interchange/Core/Private/Tasks/InterchangeTaskCompletion.cpp
+++ b/Engine/Source/Runtime/Interchange/Core/Private/Tasks/InterchangeTaskCompletion.cpp
@@ -19,7 +19,7 @@ void Interchange::FTaskCompletion::DoTask(ENamedThreads::Type CurrentThread, con
 {
 	TSharedPtr<Interchange::FImportAsyncHelper, ESPMode::ThreadSafe> AsyncHelper = WeakAsyncHelper.Pin();
 	check(AsyncHelper.IsValid());
-
+	bool bIsFutureRootObjectSet = false;
 	for(TPair<int32, TArray<Interchange::FImportAsyncHelper::FImportedAssetInfo>>& AssetInfosPerSourceIndexPair : AsyncHelper->ImportedAssetsPerSourceIndex)
 	{
 		const int32 SourceIndex = AssetInfosPerSourceIndexPair.Key;
@@ -58,8 +58,9 @@ void Interchange::FTaskCompletion::DoTask(ENamedThreads::Type CurrentThread, con
 			FAssetRegistryModule::AssetCreated(Asset);
 #endif //WITH_ENGINE
 
-			if (SourceIndex == 0)
+			if (!bIsFutureRootObjectSet && SourceIndex == 0)
 			{
+				bIsFutureRootObjectSet = true;
 				AsyncHelper->RootObject.SetValue(Asset);
 				AsyncHelper->RootObjectCompletionEvent->DispatchSubsequents();
 			}
diff --git a/Engine/Source/Runtime/Interchange/Core/Private/Tasks/InterchangeTaskCreateAsset.cpp b/Engine/Source/Runtime/Interchange/Core/Private/Tasks/InterchangeTaskCreateAsset.cpp
--- a/Engine/Source/Runtime/Interchange/Core/Private/Tasks/InterchangeTaskCreateAsset.cpp
+++ b/Engine/Source/Runtime/Interchange/Core/Private/Tasks/InterchangeTaskCreateAsset.cpp
@@ -16,11 +16,11 @@
 #include "UObject/UObjectGlobals.h"
 #include "UObject/ObjectMacros.h"
 #include "UObject/WeakObjectPtrTemplates.h"
-#include "Nodes/BaseNode.h"
+#include "Nodes/InterchangeBaseNode.h"
 
 namespace Interchange_InternalImplementation
 {
-	void InternalGetPackageName(const Interchange::FImportAsyncHelper& AsyncHelper, const int32 SourceIndex, const FString& PackageBasePath, const Interchange::FBaseNode* Node, FString& OutPackageName, FString& OutAssetName)
+	void InternalGetPackageName(const Interchange::FImportAsyncHelper& AsyncHelper, const int32 SourceIndex, const FString& PackageBasePath, const UInterchangeBaseNode* Node, FString& OutPackageName, FString& OutAssetName)
 	{
 		const UInterchangeSourceData* SourceData = AsyncHelper.SourceDatas[SourceIndex];
 		check(SourceData);
@@ -88,6 +88,10 @@ void Interchange::FTaskCreatePackage::DoTask(ENamedThreads::Type CurrentThread,
 		CreateAssetParams.Parent = Pkg;
 		CreateAssetParams.SourceData = AsyncHelper->SourceDatas[SourceIndex];
 		CreateAssetParams.Translator = nullptr;
+		if (AsyncHelper->BaseNodeContainers.IsValidIndex(SourceIndex))
+		{
+			CreateAssetParams.NodeContainer = AsyncHelper->BaseNodeContainers[SourceIndex].Get();
+		}
 		CreateAssetParams.ReimportObject = AsyncHelper->TaskData.ReimportObject;
 		//Make sure the asset UObject is created with the correct type on the main thread
 		Factory->CreateEmptyAsset(CreateAssetParams);
@@ -113,25 +117,27 @@ void Interchange::FTaskCreateAsset::DoTask(ENamedThreads::Type CurrentThread, co
 		PackageName = Pkg->GetPathName();
 	}
 
-	FScopeLock Lock(&AsyncHelper->CreatedPackagesLock);
-	UPackage** PkgPtr = AsyncHelper->CreatedPackages.Find(PackageName);
-	Lock.Unlock();
-
-	if (!PkgPtr || !(*PkgPtr))
+	UPackage* Pkg = nullptr;
 	{
-		const FText Message = FText::Format(NSLOCTEXT("Interchange", "CannotCreateAssetNoPackageErrorMsg", "Cannot create asset named '{1}', package '{0}'was not created properly."), FText::FromString(PackageName), FText::FromString(AssetName));
-		UE_LOG(LogInterchangeCore, Warning, TEXT("%s"), *Message.ToString());
-		return;
-	}
+		FScopeLock Lock(&AsyncHelper->CreatedPackagesLock);
+		UPackage** PkgPtr = AsyncHelper->CreatedPackages.Find(PackageName);
 
-	if (!AsyncHelper->SourceDatas.IsValidIndex(SourceIndex) || !AsyncHelper->Translators.IsValidIndex(SourceIndex))
-	{
-		const FText Message = FText::Format(NSLOCTEXT("Interchange", "CannotCreateAssetMissingDataErrorMsg", "Cannot create asset named '{0}', Source data or translator is invalid."), FText::FromString(AssetName));
-		UE_LOG(LogInterchangeCore, Warning, TEXT("%s"), *Message.ToString());
-		return;
-	}
+		if (!PkgPtr || !(*PkgPtr))
+		{
+			const FText Message = FText::Format(NSLOCTEXT("Interchange", "CannotCreateAssetNoPackageErrorMsg", "Cannot create asset named '{1}', package '{0}'was not created properly."), FText::FromString(PackageName), FText::FromString(AssetName));
+			UE_LOG(LogInterchangeCore, Warning, TEXT("%s"), *Message.ToString());
+			return;
+		}
+
+		if (!AsyncHelper->SourceDatas.IsValidIndex(SourceIndex) || !AsyncHelper->Translators.IsValidIndex(SourceIndex))
+		{
+			const FText Message = FText::Format(NSLOCTEXT("Interchange", "CannotCreateAssetMissingDataErrorMsg", "Cannot create asset named '{0}', Source data or translator is invalid."), FText::FromString(AssetName));
+			UE_LOG(LogInterchangeCore, Warning, TEXT("%s"), *Message.ToString());
+			return;
+		}
 
-	UPackage* Pkg = *PkgPtr;
+		Pkg = *PkgPtr;
+	}
 
 	UInterchangeTranslatorBase* Translator = AsyncHelper->Translators[SourceIndex];
 	//Import Asset describe by the node
@@ -141,14 +147,20 @@ void Interchange::FTaskCreateAsset::DoTask(ENamedThreads::Type CurrentThread, co
 	CreateAssetParams.Parent = Pkg;
 	CreateAssetParams.SourceData = AsyncHelper->SourceDatas[SourceIndex];
 	CreateAssetParams.Translator = Translator;
+	if (AsyncHelper->BaseNodeContainers.IsValidIndex(SourceIndex))
+	{
+		CreateAssetParams.NodeContainer = AsyncHelper->BaseNodeContainers[SourceIndex].Get();
+	}
 	CreateAssetParams.ReimportObject = AsyncHelper->TaskData.ReimportObject;
 
 	UObject* NodeAsset = Factory->CreateAsset(CreateAssetParams);
 	if (NodeAsset)
 	{
+		FScopeLock Lock(&AsyncHelper->ImportedAssetsPerSourceIndexLock);
 		TArray<Interchange::FImportAsyncHelper::FImportedAssetInfo>& ImportedInfos = AsyncHelper->ImportedAssetsPerSourceIndex.FindOrAdd(SourceIndex);
 		Interchange::FImportAsyncHelper::FImportedAssetInfo& AssetInfo = ImportedInfos.AddZeroed_GetRef();
 		AssetInfo.ImportAsset = NodeAsset;
+		Node->ReferenceObject = FSoftObjectPath(NodeAsset);
 		AssetInfo.Factory = Factory;
 	}
 }
\ No newline at end of file
diff --git a/Engine/Source/Runtime/Interchange/Core/Private/Tasks/InterchangeTaskCreateAsset.h b/Engine/Source/Runtime/Interchange/Core/Private/Tasks/InterchangeTaskCreateAsset.h
--- a/Engine/Source/Runtime/Interchange/Core/Private/Tasks/InterchangeTaskCreateAsset.h
+++ b/Engine/Source/Runtime/Interchange/Core/Private/Tasks/InterchangeTaskCreateAsset.h
@@ -8,7 +8,7 @@
 #include "InterchangeManager.h"
 #include "Stats/Stats.h"
 #include "UObject/WeakObjectPtrTemplates.h"
-#include "Nodes/BaseNode.h"
+#include "Nodes/InterchangeBaseNode.h"
 
 namespace Interchange
 {
@@ -22,11 +22,11 @@ private:
 	FString PackageBasePath;
 	int32 SourceIndex;
 	TWeakPtr<Interchange::FImportAsyncHelper, ESPMode::ThreadSafe> WeakAsyncHelper;
-	const FBaseNode* Node;
+	const UInterchangeBaseNode* Node;
 	UInterchangeFactoryBase* Factory;
 
 public:
-	FTaskCreatePackage(const FString& InPackageBasePath, const int32 InSourceIndex, TWeakPtr<Interchange::FImportAsyncHelper, ESPMode::ThreadSafe> InAsyncHelper, const Interchange::FBaseNode* InNode, UInterchangeFactoryBase* InFactory)
+	FTaskCreatePackage(const FString& InPackageBasePath, const int32 InSourceIndex, TWeakPtr<Interchange::FImportAsyncHelper, ESPMode::ThreadSafe> InAsyncHelper, const UInterchangeBaseNode* InNode, UInterchangeFactoryBase* InFactory)
 		: PackageBasePath(InPackageBasePath)
 		, SourceIndex(InSourceIndex)
 		, WeakAsyncHelper(InAsyncHelper)
@@ -66,11 +66,11 @@ private:
 	FString PackageBasePath;
 	int32 SourceIndex;
 	TWeakPtr<Interchange::FImportAsyncHelper, ESPMode::ThreadSafe> WeakAsyncHelper;
-	const FBaseNode* Node;
+	const UInterchangeBaseNode* Node;
 	UInterchangeFactoryBase* Factory;
 
 public:
-	FTaskCreateAsset(const FString& InPackageBasePath, const int32 InSourceIndex, TWeakPtr<Interchange::FImportAsyncHelper, ESPMode::ThreadSafe> InAsyncHelper, const Interchange::FBaseNode* InNode, UInterchangeFactoryBase* InFactory)
+	FTaskCreateAsset(const FString& InPackageBasePath, const int32 InSourceIndex, TWeakPtr<Interchange::FImportAsyncHelper, ESPMode::ThreadSafe> InAsyncHelper, const UInterchangeBaseNode* InNode, UInterchangeFactoryBase* InFactory)
 		: PackageBasePath(InPackageBasePath)
 		, SourceIndex(InSourceIndex)
 		, WeakAsyncHelper(InAsyncHelper)
diff --git a/Engine/Source/Runtime/Interchange/Core/Private/Tasks/InterchangeTaskParsing.cpp b/Engine/Source/Runtime/Interchange/Core/Private/Tasks/InterchangeTaskParsing.cpp
--- a/Engine/Source/Runtime/Interchange/Core/Private/Tasks/InterchangeTaskParsing.cpp
+++ b/Engine/Source/Runtime/Interchange/Core/Private/Tasks/InterchangeTaskParsing.cpp
@@ -17,8 +17,8 @@
 #include "UObject/UObjectGlobals.h"
 #include "UObject/ObjectMacros.h"
 #include "UObject/WeakObjectPtrTemplates.h"
-#include "Nodes/BaseNode.h"
-#include "Nodes/BaseNodeContainer.h"
+#include "Nodes/InterchangeBaseNode.h"
+#include "Nodes/InterchangeBaseNodeContainer.h"
 
 void Interchange::FTaskParsing::DoTask(ENamedThreads::Type CurrentThread, const FGraphEventRef& MyCompletionGraphEvent)
 {
@@ -28,10 +28,10 @@ void Interchange::FTaskParsing::DoTask(ENamedThreads::Type CurrentThread, const
 	//Parse each graph and prepare import task data, we will then be able to create all the task with the correct dependencies
 	struct FTaskData
 	{
-		Interchange::FNodeUniqueID UniqueID;
+		FName UniqueID;
 		int32 SourceIndex;
-		const Interchange::FBaseNode* Node;
-		TArray<Interchange::FNodeUniqueID> Dependencies;
+		const UInterchangeBaseNode* Node;
+		TArray<FName> Dependencies;
 		FGraphEventRef GraphEventRef;
 		FGraphEventArray Prerequistes;
 		UInterchangeFactoryBase* Factory;
@@ -43,7 +43,16 @@ void Interchange::FTaskParsing::DoTask(ENamedThreads::Type CurrentThread, const
 	TArray<FTaskData> TaskDatas;
 	for (int32 SourceIndex = 0; SourceIndex < AsyncHelper->SourceDatas.Num(); ++SourceIndex)
 	{
-		AsyncHelper->BaseNodeContainers[SourceIndex].IterateNodes([&](const Interchange::FNodeUniqueID NodeUID, const Interchange::FBaseNode* Node)
+		if (!AsyncHelper->BaseNodeContainers.IsValidIndex(SourceIndex))
+		{
+			continue;
+		}
+		UInterchangeBaseNodeContainer* BaseNodeContainer = AsyncHelper->BaseNodeContainers[SourceIndex].Get();
+		if (!BaseNodeContainer)
+		{
+			continue;
+		}
+		BaseNodeContainer->IterateNodes([&](const FName NodeUID, const UInterchangeBaseNode* Node)
 		{
 			if (Node->GetAssetClass() != nullptr)
 			{
diff --git a/Engine/Source/Runtime/Interchange/Core/Private/Tasks/InterchangeTaskPipeline.cpp b/Engine/Source/Runtime/Interchange/Core/Private/Tasks/InterchangeTaskPipeline.cpp
--- a/Engine/Source/Runtime/Interchange/Core/Private/Tasks/InterchangeTaskPipeline.cpp
+++ b/Engine/Source/Runtime/Interchange/Core/Private/Tasks/InterchangeTaskPipeline.cpp
@@ -8,7 +8,7 @@
 #include "Stats/Stats.h"
 #include "Templates/SharedPointer.h"
 #include "UObject/WeakObjectPtrTemplates.h"
-#include "Nodes/BaseNodeContainer.h"
+#include "Nodes/InterchangeBaseNodeContainer.h"
 
 void Interchange::FTaskPipeline::DoTask(ENamedThreads::Type CurrentThread, const FGraphEventRef& MyCompletionGraphEvent)
 {
@@ -19,8 +19,8 @@ void Interchange::FTaskPipeline::DoTask(ENamedThreads::Type CurrentThread, const
 	{
 		for (int32 GraphIndex = 0; GraphIndex < AsyncHelper->BaseNodeContainers.Num(); ++GraphIndex)
 		{
-			check(AsyncHelper->BaseNodeContainerAdapters[GraphIndex].IsValid());
-			Pipeline->ScriptedExecuteImportPipeline(AsyncHelper->BaseNodeContainerAdapters[GraphIndex].Get());
+			check(AsyncHelper->BaseNodeContainers[GraphIndex].IsValid());
+			Pipeline->ScriptedExecuteImportPipeline(AsyncHelper->BaseNodeContainers[GraphIndex].Get());
 		}
 	}
 }
diff --git a/Engine/Source/Runtime/Interchange/Core/Private/Tasks/InterchangeTaskTranslator.cpp b/Engine/Source/Runtime/Interchange/Core/Private/Tasks/InterchangeTaskTranslator.cpp
--- a/Engine/Source/Runtime/Interchange/Core/Private/Tasks/InterchangeTaskTranslator.cpp
+++ b/Engine/Source/Runtime/Interchange/Core/Private/Tasks/InterchangeTaskTranslator.cpp
@@ -9,7 +9,7 @@
 #include "Stats/Stats.h"
 #include "Templates/SharedPointer.h"
 #include "UObject/WeakObjectPtrTemplates.h"
-#include "Nodes/BaseNodeContainer.h"
+#include "Nodes/InterchangeBaseNodeContainer.h"
 
 void Interchange::FTaskTranslator::DoTask(ENamedThreads::Type CurrentThread, const FGraphEventRef& MyCompletionGraphEvent)
 {
@@ -27,12 +27,12 @@ void Interchange::FTaskTranslator::DoTask(ENamedThreads::Type CurrentThread, con
 		return;
 	}
 
-	if (!AsyncHelper->BaseNodeContainers.IsValidIndex(SourceIndex))
+	if (!AsyncHelper->BaseNodeContainers.IsValidIndex(SourceIndex) || !AsyncHelper->BaseNodeContainers[SourceIndex].IsValid())
 	{
 		return;
 	}
 
 	//Translate the source data
-	Interchange::FBaseNodeContainer& BaseNodeContainer = AsyncHelper->BaseNodeContainers[SourceIndex];
+	UInterchangeBaseNodeContainer& BaseNodeContainer = *(AsyncHelper->BaseNodeContainers[SourceIndex].Get());
 	Translator->Translate(SourceData, BaseNodeContainer);
 }
diff --git a/Engine/Source/Runtime/Interchange/Core/Private/Types/AttributeStorage.cpp b/Engine/Source/Runtime/Interchange/Core/Private/Types/AttributeStorage.cpp
--- a/Engine/Source/Runtime/Interchange/Core/Private/Types/AttributeStorage.cpp
+++ b/Engine/Source/Runtime/Interchange/Core/Private/Types/AttributeStorage.cpp
@@ -61,6 +61,19 @@ namespace Interchange
 		}
 	}
 
+	FAttributeStorage::FAttributeStorage(const FAttributeStorage& Other)
+	{
+		//Lock both Storage mutex
+		FScopeLock ScopeLock(&StorageMutex);
+		FScopeLock ScopeLockOther(&Other.StorageMutex);
+
+		//Copy all value
+		AttributeAllocationTable = Other.AttributeAllocationTable;
+		AttributeStorage = Other.AttributeStorage;
+		FragmentedMemoryCost = Other.FragmentedMemoryCost;
+		DefragRatio = Other.DefragRatio;
+	}
+
 	FAttributeStorage& FAttributeStorage::operator=(const FAttributeStorage& Other)
 	{
 		//Lock both Storage mutex
diff --git a/Engine/Source/Runtime/Interchange/Core/Public/InterchangeFactoryBase.h b/Engine/Source/Runtime/Interchange/Core/Public/InterchangeFactoryBase.h
--- a/Engine/Source/Runtime/Interchange/Core/Public/InterchangeFactoryBase.h
+++ b/Engine/Source/Runtime/Interchange/Core/Public/InterchangeFactoryBase.h
@@ -3,7 +3,7 @@
 #pragma once
 
 #include "CoreMinimal.h"
-#include "Nodes/BaseNode.h"
+#include "Nodes/InterchangeBaseNode.h"
 #include "InterchangeTranslatorBase.h"
 #include "UObject/Class.h"
 #include "UObject/Object.h"
@@ -38,7 +38,7 @@ public:
 		FString AssetName = FString();
 
 		/** The base node that describe how to create the asset */
-		const Interchange::FBaseNode* AssetNode = nullptr;
+		const UInterchangeBaseNode* AssetNode = nullptr;
 
 		/** The translator is use to retrieve the PayLoad data in case the factory need it */
 		const UInterchangeTranslatorBase* Translator = nullptr;
@@ -46,6 +46,13 @@ public:
 		/** The source data, mainly use to set the asset import data file. TODO: we have to refactor UAssetImportData, the source data should be the base class for this now */
 		const UInterchangeSourceData* SourceData = nullptr;
 
+		/** The node container associate with the current source index */
+		const UInterchangeBaseNodeContainer* NodeContainer = nullptr;
+
+		/**
+		 * If when we try to create the package we found out the asset already exist, this field will contain
+		 * the asset we want to re-import. The re-import should just change the source data and not any asset settings.
+		 */
 		UObject* ReimportObject = nullptr;
 	};
 
diff --git a/Engine/Source/Runtime/Interchange/Core/Public/InterchangeManager.h b/Engine/Source/Runtime/Interchange/Core/Public/InterchangeManager.h
--- a/Engine/Source/Runtime/Interchange/Core/Public/InterchangeManager.h
+++ b/Engine/Source/Runtime/Interchange/Core/Public/InterchangeManager.h
@@ -10,7 +10,7 @@
 #include "InterchangeSourceData.h"
 #include "InterchangeTranslatorBase.h"
 #include "InterchangeWriterBase.h"
-#include "Nodes/InterchangeBaseNodeContainerAdapter.h"
+#include "Nodes/InterchangeBaseNodeContainer.h"
 #include "UObject/Package.h"
 #include "UObject/Object.h"
 #include "UObject/ObjectMacros.h"
@@ -73,8 +73,7 @@ namespace Interchange
 		/* FGCObject interface */
 		virtual void AddReferencedObjects(FReferenceCollector& Collector) override;
 
-		TArray<TStrongObjectPtr<UInterchangeBaseNodeContainerAdapter>> BaseNodeContainerAdapters;
-		TArray <FBaseNodeContainer> BaseNodeContainers;
+		TArray <TStrongObjectPtr<UInterchangeBaseNodeContainer>> BaseNodeContainers;
 
 		TArray<UInterchangeSourceData* > SourceDatas;
 		TArray<UInterchangeTranslatorBase* > Translators;
@@ -96,8 +95,10 @@ namespace Interchange
 		{
 			UObject* ImportAsset;
 			UInterchangeFactoryBase* Factory;
+			UInterchangeBaseNode* AssetNode;
 		};
 
+		FCriticalSection ImportedAssetsPerSourceIndexLock;
 		TMap<int32, TArray<FImportedAssetInfo>> ImportedAssetsPerSourceIndex;
 
 		FImportAsyncHelperData TaskData;
@@ -302,8 +303,6 @@ public:
 	/** Delete the specified AsyncHelper and remove it from the array that was holding it. */
 	void ReleaseAsyncHelper(TWeakPtr<Interchange::FImportAsyncHelper, ESPMode::ThreadSafe> AsyncHelper);
 
-protected:
-
 	/*
 	 * Return the first translator that can translate the source data.
 	 * @Param SourceData - The source data for which we search a translator.
@@ -311,6 +310,8 @@ protected:
 	 */
 	UInterchangeTranslatorBase* GetTranslatorForSourceData(const UInterchangeSourceData* SourceData) const;
 
+protected:
+
 	/** Return true if we can show some UI */
 	static bool IsAttended();
 
diff --git a/Engine/Source/Runtime/Interchange/Core/Public/InterchangePipelineBase.h b/Engine/Source/Runtime/Interchange/Core/Public/InterchangePipelineBase.h
--- a/Engine/Source/Runtime/Interchange/Core/Public/InterchangePipelineBase.h
+++ b/Engine/Source/Runtime/Interchange/Core/Public/InterchangePipelineBase.h
@@ -3,7 +3,7 @@
 #pragma once
 
 #include "CoreMinimal.h"
-#include "Nodes/BaseNodeContainer.h"
+#include "Nodes/InterchangeBaseNodeContainer.h"
 #include "InterchangeSourceData.h"
 #include "UObject/Object.h"
 #include "UObject/ObjectMacros.h"
@@ -22,9 +22,9 @@ public:
 	 * the Interchange manager is calling this function not the virtual one that is call by the default implementation.
 	 */
 	UFUNCTION(BlueprintNativeEvent, BlueprintCallable, Category = "Interchange | Translator")
-	bool ScriptedExecuteImportPipeline(UInterchangeBaseNodeContainerAdapter* BaseNodeContainerAdapter);
+	bool ScriptedExecuteImportPipeline(UInterchangeBaseNodeContainer* BaseNodeContainerAdapter);
 	/** The default implementation (call if the blueprint do not have any implementation) will call the virtual ExecuteImportPipeline */
-	bool ScriptedExecuteImportPipeline_Implementation(UInterchangeBaseNodeContainerAdapter* BaseNodeContainerAdapter)
+	bool ScriptedExecuteImportPipeline_Implementation(UInterchangeBaseNodeContainer* BaseNodeContainerAdapter)
 	{
 		//By default we call the virtual import pipeline execution
 		return ExecuteImportPipeline(BaseNodeContainerAdapter);
@@ -35,10 +35,10 @@ public:
 	 * the Interchange manager is calling this function not the virtual one that is call by the default implementation.
 	 */
 	UFUNCTION(BlueprintCallable, BlueprintNativeEvent, Category = "Interchange | Translator")
-	bool ScriptedExecuteExportPipeline(UInterchangeBaseNodeContainerAdapter* BaseNodeContainerAdapter);
+	bool ScriptedExecuteExportPipeline(UInterchangeBaseNodeContainer* BaseNodeContainerAdapter);
 
 	/** The default implementation (call if the blueprint do not have any implementation) will call the virtual ExecuteExportPipeline */
-	bool ScriptedExecuteExportPipeline_Implementation(UInterchangeBaseNodeContainerAdapter* BaseNodeContainerAdapter)
+	bool ScriptedExecuteExportPipeline_Implementation(UInterchangeBaseNodeContainer* BaseNodeContainerAdapter)
 	{
 		return ExecuteExportPipeline(BaseNodeContainerAdapter);
 	}
@@ -62,7 +62,7 @@ protected:
 	 * The interchange manager is not calling this function directly. It is calling the blueprint native event in case this object is a blueprint derive object.
 	 * By default the scripted implementation is calling this virtual pipeline.
 	 */
-	virtual bool ExecuteImportPipeline(UInterchangeBaseNodeContainerAdapter* BaseNodeContainerAdapter)
+	virtual bool ExecuteImportPipeline(UInterchangeBaseNodeContainer* BaseNodeContainerAdapter)
 	{
 		return false;
 	}
@@ -79,7 +79,7 @@ protected:
 	}
 
 	/** This function can modify the BaseNodeContainer to create a pipeline that will set/validate the graph nodes hierarchy and options.*/
-	virtual bool ExecuteExportPipeline(UInterchangeBaseNodeContainerAdapter* BaseNodeContainerAdapter)
+	virtual bool ExecuteExportPipeline(UInterchangeBaseNodeContainer* BaseNodeContainerAdapter)
 	{
 		return false;
 	}
diff --git a/Engine/Source/Runtime/Interchange/Core/Public/InterchangeTranslatorBase.h b/Engine/Source/Runtime/Interchange/Core/Public/InterchangeTranslatorBase.h
--- a/Engine/Source/Runtime/Interchange/Core/Public/InterchangeTranslatorBase.h
+++ b/Engine/Source/Runtime/Interchange/Core/Public/InterchangeTranslatorBase.h
@@ -3,7 +3,7 @@
 #pragma once
 
 #include "CoreMinimal.h"
-#include "Nodes/BaseNodeContainer.h"
+#include "Nodes/InterchangeBaseNodeContainer.h"
 #include "InterchangeSourceData.h"
 #include "UObject/Object.h"
 #include "UObject/ObjectMacros.h"
@@ -29,7 +29,7 @@ public:
 	 * @param BaseNodeContainer - The nodes container where to put the translated source data result.
 	 * @return true if the translator can translate the source data, false otherwise.
 	 */
-	virtual bool Translate(const UInterchangeSourceData* SourceData, Interchange::FBaseNodeContainer& BaseNodeContainer) const
+	virtual bool Translate(const UInterchangeSourceData* SourceData, UInterchangeBaseNodeContainer& BaseNodeContainer) const
 	{
 		return false;
 	}
diff --git a/Engine/Source/Runtime/Interchange/Core/Public/InterchangeWriterBase.h b/Engine/Source/Runtime/Interchange/Core/Public/InterchangeWriterBase.h
--- a/Engine/Source/Runtime/Interchange/Core/Public/InterchangeWriterBase.h
+++ b/Engine/Source/Runtime/Interchange/Core/Public/InterchangeWriterBase.h
@@ -3,7 +3,7 @@
 #pragma once
 
 #include "CoreMinimal.h"
-#include "Nodes/BaseNodeContainer.h"
+#include "Nodes/InterchangeBaseNodeContainer.h"
 #include "UObject/Object.h"
 #include "UObject/ObjectMacros.h"
 
@@ -20,5 +20,5 @@ public:
 	 * @param BaseNodeContainer - The container holding nodes that describe what to export
 	 * @return true if the writer can export the Nodes, false otherwise.
 	 */
-	virtual bool Export(Interchange::FBaseNodeContainer& BaseNodeContainer) const {return false;}
+	virtual bool Export(UInterchangeBaseNodeContainer* BaseNodeContainer) const {return false;}
 };
diff --git a/Engine/Source/Runtime/Interchange/Core/Public/Nodes/BaseNode.h b/Engine/Source/Runtime/Interchange/Core/Public/Nodes/BaseNode.h
deleted file mode 100644
diff --git a/Engine/Source/Runtime/Interchange/Core/Public/Nodes/BaseNodeContainer.h b/Engine/Source/Runtime/Interchange/Core/Public/Nodes/BaseNodeContainer.h
deleted file mode 100644
diff --git a/Engine/Source/Runtime/Interchange/Core/Public/Nodes/InterchangeBaseNode.h b/Engine/Source/Runtime/Interchange/Core/Public/Nodes/InterchangeBaseNode.h
new file mode 100644
--- /dev/null
+++ b/Engine/Source/Runtime/Interchange/Core/Public/Nodes/InterchangeBaseNode.h
@@ -0,0 +1,359 @@
+// Copyright Epic Games, Inc. All Rights Reserved.
+
+#pragma once
+
+#include "CoreMinimal.h"
+#include "Types/AttributeStorage.h"
+#include "UObject/Class.h"
+#include "UObject/Object.h"
+#include "UObject/ObjectMacros.h"
+
+#include "InterchangeBaseNode.generated.h"
+
+/**
+ * Internal Helper to get set custom property for class that derive from UInterchangeBaseNode. This is use by the macro IMPLEMENT_UOD_ATTRIBUTE.
+ */
+namespace InterchangePrivateNodeBase
+{
+	/**
+	 * Retrieve a custom attribute if the attribute exist
+	 *
+	 * @param Attributes - The attribute storage you want to query the custom attribute
+	 * @param AttributeKey - The storage key for the attribute
+	 * @param OperationName - The name of the operation in case there is an error
+	 * @param OutAttributeValue - This is where we store the value we retrieve from the storage
+	 *
+	 * @return - return true if the attribute exist in the storage and was query without error.
+	 *           return false if the attribute do not exist or there is an error retriving it from the Storage
+	 */
+	template<typename ValueType>
+	bool GetCustomAttribute(const Interchange::FAttributeStorage& Attributes, const Interchange::FAttributeKey& AttributeKey, const FString& OperationName, ValueType& OutAttributeValue)
+	{
+		if (!Attributes.ContainAttribute(AttributeKey))
+		{
+			return false;
+		}
+		Interchange::FAttributeStorage::TAttributeHandle<ValueType> AttributeHandle = Attributes.GetAttributeHandle<ValueType>(AttributeKey);
+		if (!AttributeHandle.IsValid())
+		{
+			return false;
+		}
+		Interchange::EAttributeStorageResult Result = AttributeHandle.Get(OutAttributeValue);
+		if (!Interchange::IsAttributeStorageResultSuccess(Result))
+		{
+			Interchange::LogAttributeStorageErrors(Result, OperationName, AttributeKey);
+			return false;
+		}
+		return true;
+	}
+
+	/**
+	 * Add or update a custom attribute value in the specified storage
+	 *
+	 * @param Attributes - The attribute storage you want to add or update the custom attribute
+	 * @param AttributeKey - The storage key for the attribute
+	 * @param OperationName - The name of the operation in case there is an error
+	 * @param AttributeValue - The value we want to add or update in the storage
+	 */
+	template<typename ValueType>
+	bool SetCustomAttribute(Interchange::FAttributeStorage& Attributes, const Interchange::FAttributeKey& AttributeKey, const FString& OperationName, const ValueType& AttributeValue)
+	{
+		Interchange::EAttributeStorageResult Result = Attributes.RegisterAttribute(AttributeKey, AttributeValue);
+		if (!Interchange::IsAttributeStorageResultSuccess(Result))
+		{
+			Interchange::LogAttributeStorageErrors(Result, OperationName, AttributeKey);
+			return false;
+		}
+		return true;
+	}
+}
+
+/**
+ * Use those macro to create Get/Set/ApplyToAsset for an attribute you want to support in your custom UOD node that derive from UInterchangeBaseNode.
+ * The attribute key will be declare under private access modifier and the functions are declare under public access modifier.
+ * So the class access modifier is public after calling this macro.
+ *
+ * @note - The Get will return false if the attribute was never set.
+ * @note - The Set will add the attribute to the storage or update the value if the storage already has this attribute.
+ * @note - The Apply will set a member variable of the AssetType instance. It will apply the value only if the storage contain the attribute key
+ *
+ * @param AttributeName - The name of the Get/Set functions. Example if you pass Foo you will end up with GetFoo and SetFoo function
+ * @param AttributeType - This is to specify the type of the attribute. bool, float, FString... anything supported by the FAttributeStorage
+ * @param AssetType - This is the asset you want to apply the storage value"
+ * @param EnumType - Optional, specify it only if the AssetType member is an enum so we can type cast it in the apply function (we use uint8 to store the enum value)"
+ */
+#define IMPLEMENT_NODE_ATTRIBUTE_KEY(AttributeName)																		\
+const Interchange::FAttributeKey Macro_Custom##AttributeName##Key = Interchange::FAttributeKey(TEXT(#AttributeName));	\
+
+#if WITH_ENGINE
+#define IMPLEMENT_NODE_ATTRIBUTE_APPLY_UOBJECT(AttributeName, AttributeType, AssetType, EnumType)	\
+bool ApplyCustom##AttributeName##ToAsset(UObject* Asset) const										\
+{																									\
+	if (!Asset)																						\
+	{																								\
+		return false;																				\
+	}																								\
+	AssetType* TypedObject = Cast<AssetType>(Asset);												\
+	if (!TypedObject)																				\
+	{																								\
+		return false;																				\
+	}																								\
+	AttributeType ValueData;																		\
+	if (GetCustom##AttributeName(ValueData))														\
+	{																								\
+		TypedObject->AttributeName = EnumType(ValueData);											\
+		return true;																				\
+	}																								\
+	return false;																					\
+}
+
+#else //#if WITH_ENGINE
+
+#define IMPLEMENT_NODE_ATTRIBUTE_APPLY_UOBJECT(AttributeName, AttributeType, AssetType, EnumType)
+
+#endif //#if WITH_ENGINE
+
+#define IMPLEMENT_NODE_ATTRIBUTE_GETTER(AttributeName, AttributeType)																					\
+	FString OperationName = GetTypeName() + TEXT(".Get##AttributeName");																				\
+	return InterchangePrivateNodeBase::GetCustomAttribute<AttributeType>(Attributes, Macro_Custom##AttributeName##Key, OperationName, AttributeValue);
+
+#define IMPLEMENT_NODE_ATTRIBUTE_SETTER_NODELEGATE(AttributeName, AttributeType)															\
+	FString OperationName = GetTypeName() + TEXT(".Set##AttributeName");																				\
+	return InterchangePrivateNodeBase::SetCustomAttribute<AttributeType>(Attributes, Macro_Custom##AttributeName##Key, OperationName, AttributeValue);
+
+#if WITH_ENGINE
+
+#define IMPLEMENT_NODE_ATTRIBUTE_SETTER(NodeClassName, AttributeName, AttributeType, AssetType)														\
+	FString OperationName = GetTypeName() + TEXT(".Set##AttributeName");																			\
+	if(InterchangePrivateNodeBase::SetCustomAttribute<AttributeType>(Attributes, Macro_Custom##AttributeName##Key, OperationName, AttributeValue))	\
+	{																																				\
+		if(bAddApplyDelegate)																														\
+		{																																			\
+			TArray<Interchange::FApplyAttributeToAsset>& Delegates = ApplyCustomAttributeDelegates.FindOrAdd(AssetType::StaticClass());				\
+			Delegates.Add(Interchange::FApplyAttributeToAsset::CreateUObject(this, &NodeClassName::ApplyCustom##AttributeName##ToAsset));			\
+		}																																			\
+		return true;																																\
+	}																																				\
+	return false;
+
+#else //#if WITH_ENGINE
+
+#define IMPLEMENT_NODE_ATTRIBUTE_SETTER(NodeClassName, AttributeName, AttributeType, AssetType)															\
+	FString OperationName = GetTypeName() + TEXT(".Set##AttributeName");																				\
+	return InterchangePrivateNodeBase::SetCustomAttribute<AttributeType>(Attributes, Macro_Custom##AttributeName##Key, OperationName, AttributeValue);
+
+#endif //#if WITH_ENGINE
+
+
+//Interchange namespace
+namespace Interchange
+{
+
+DECLARE_DELEGATE_RetVal_OneParam(bool, FApplyAttributeToAsset, UObject*);
+
+/**
+ * Helper struct use to declare static const data we use in the UInterchangeBaseNode
+ * Node that derive from UInterchangeBaseNode can also add a struct that derive from this one to add there static data
+ * @note: The static data are mainly for holding Attribute keys. All attributes that are always available for a node should be in this class or a derived class.
+ */
+struct FBaseNodeStaticData
+{
+	static const FAttributeKey& UniqueIDKey()
+	{
+		static FAttributeKey AttributeKey(TEXT("__UNQ_ID_"));
+		return AttributeKey;
+	}
+
+	static const FAttributeKey& DisplayLabelKey()
+	{
+		static FAttributeKey AttributeKey(TEXT("__DSPL_LBL_"));
+		return AttributeKey;
+	}
+
+	static const FAttributeKey& ParentIDKey()
+	{
+		static FAttributeKey AttributeKey(TEXT("__PARENT_UID_"));
+		return AttributeKey;
+	}
+
+	static const FAttributeKey& IsEnabledKey()
+	{
+		static FAttributeKey AttributeKey(TEXT("__IS_NBLD_"));
+		return AttributeKey;
+	}
+
+	static const FAttributeKey& DependencyCountKey()
+	{
+		static FAttributeKey AttributeKey(TEXT("__DEPENDENCY_COUNT_"));
+		return AttributeKey;
+	}
+
+	static const FAttributeKey& DependencyBaseKey()
+	{
+		static FAttributeKey AttributeKey(TEXT("__DEPENDENCY_INDEX_"));
+		return AttributeKey;
+	}
+};
+
+} //ns Interchange
+
+/**
+ * This struct is used to store and retrieve key value attributes. The attributes are store in a generic FAttributeStorage which serialize the value in a TArray64<uint8>
+ * See Interchange::EAttributeTypes to know the supported template types
+ * This is an abstract class. This is the base class of the interchange node graph format, all class in this format should derive from this class
+ */
+UCLASS(BlueprintType)
+class INTERCHANGECORE_API UInterchangeBaseNode : public UObject
+{
+	GENERATED_BODY()
+
+public:
+	UInterchangeBaseNode()
+	{}
+
+	virtual ~UInterchangeBaseNode() = default;
+
+	/**
+	 * Initialize the base data of the node
+	 * @param UniqueID - The uniqueId for this node
+	 * @param DisplayLabel - The name of the node
+	 *
+	 */
+	UFUNCTION(BlueprintCallable, Category = "Interchange | Node")
+	void InitializeNode(const FName& UniqueID, const FName& DisplayLabel);
+
+	/**
+	 * Return the node type name of the class, we use this when reporting error
+	 */
+	virtual FString GetTypeName() const;
+
+	/**
+	 * Add an attribute to the node
+	 * @param NodeAttributeKey - The key of the attribute
+	 * @param Value - The attribute value.
+	 *
+	 */
+	template<typename T>
+	Interchange::FAttributeStorage::TAttributeHandle<T> RegisterAttribute(const Interchange::FAttributeKey& NodeAttributeKey, const T& Value)
+	{
+		const Interchange::EAttributeStorageResult Result = Attributes.RegisterAttribute(NodeAttributeKey, Value);
+		
+		if (IsAttributeStorageResultSuccess(Result))
+		{
+			return Attributes.GetAttributeHandle<T>(NodeAttributeKey);
+		}
+		LogAttributeStorageErrors(Result, TEXT("RegisterAttribute"), NodeAttributeKey);
+		return Interchange::FAttributeStorage::TAttributeHandle<T>();
+	}
+
+	/**
+	 * Return true if the node contain an attribute with the specified Key
+	 * @param nodeAttributeKey - The key of the searched attribute
+	 *
+	 */
+	virtual bool HasAttribute(const Interchange::FAttributeKey& NodeAttributeKey) const;
+
+
+	/**
+	 * This function returnan  attribute handle for the specified Key.
+	 * If there is an issue with the KEY or storage the method will trip a check, always make sure you have a valid key before calling this
+	 * 
+	 * @param NodeAttributeKey - The key of the attribute
+	 */
+	template<typename T>
+	Interchange::FAttributeStorage::TAttributeHandle<T> GetAttributeHandle(const Interchange::FAttributeKey& NodeAttributeKey) const
+	{
+		return Attributes.GetAttributeHandle<T>(NodeAttributeKey);
+	}
+
+	/**
+	 * Return the unique id pass in the constructor.
+	 *
+	 */
+	virtual FName GetUniqueID() const;
+
+	/**
+	 * Return the display label pass in the constructor.
+	 *
+	 */
+	UFUNCTION(BlueprintCallable, Category = "Interchange | Node")
+	FName GetDisplayLabel() const;
+
+	/**
+	 * Return the parent unique id. In case the attribute does not exist it will return InvalidNodeUID()
+	 */
+	UFUNCTION(BlueprintCallable, Category = "Interchange | Node")
+	FName GetParentUID() const;
+
+	/**
+	 * Set the parent unique id.
+	 */
+	UFUNCTION(BlueprintCallable, Category = "Interchange | Node")
+	bool SetParentUID(FName ParentUID);
+
+	/**
+	 * This function allow to retrieve the dependency for this object.
+	 * 
+	 */
+	UFUNCTION(BlueprintCallable, Category = "Interchange | Node")
+	void GetDependecies(TArray<FName>& OutDependencies ) const;
+
+	/**
+	 * Add one dependency to this object.
+	 */
+	UFUNCTION(BlueprintCallable, Category = "Interchange | Node")
+	bool SetDependencyUID(FName DependencyUID);
+
+	/**
+	 * IsEnable true mean that the node will be import/export, if false it will be discarded.
+	 * Return false if this node was disabled. Return true if the attribute is not there or if it was enabled.
+	 */
+	UFUNCTION(BlueprintCallable, Category = "Interchange | Node")
+	bool IsEnabled() const;
+
+	/**
+	 * Set the IsEnable attribute to determine if this node should be part of the import/export process
+	 * @param bIsEnabled - The enabled state we want to set this node. True will import/export the node, fals will not.
+	 * @return true if it was able to set the attribute, false otherwise.
+	 */
+	UFUNCTION(BlueprintCallable, Category = "Interchange | Node")
+	bool SetEnabled(const bool bIsEnabled);
+
+	/**
+	 * Return a FGuid build from the FSHA1 of all the attribute data contain in the node.
+	 *
+	 * @note the attribute are sorted by key when building the FSHA1 data. The hash will be deterministic for the same data whatever
+	 * the order we add the attributes.
+	 * This function interface is pure virtual
+	 */
+	virtual FGuid GetHash() const;
+
+	/**
+	 * Optional, Any node that can import/export an asset should return the UClass of the asset so we can find factory/writer
+	 */
+	virtual class UClass* GetAssetClass() const;
+
+	/** Return the invalid unique ID */
+	static FName InvalidNodeUID();
+
+	/**
+	 * Each Attribute that was set and have a delegate set for the specified UObject->UClass will
+	 * get the delegate execute so it apply the attribute to the UObject property.
+	 * See the macros IMPLEMENT_NODE_ATTRIBUTE_SETTER at the top of the file to know how delegates are setup for property.
+	 *
+	 */
+	void ApplyAllCustomAttributeToAsset(UObject* Object) const;
+
+	virtual void Serialize(FArchive& Ar) override;
+
+	UPROPERTY()
+	mutable FSoftObjectPath ReferenceObject;
+protected:
+	/** The storage use to store the Key value attribute for this node. */
+	Interchange::FAttributeStorage Attributes;
+
+	/* This array hold the delegate to apply the attribute that has to be set on an UObject */
+	TMap<UClass*, TArray<Interchange::FApplyAttributeToAsset>> ApplyCustomAttributeDelegates;
+
+	bool bIsInitialized = false;
+};
diff --git a/Engine/Source/Runtime/Interchange/Core/Public/Nodes/InterchangeBaseNodeContainer.h b/Engine/Source/Runtime/Interchange/Core/Public/Nodes/InterchangeBaseNodeContainer.h
new file mode 100644
--- /dev/null
+++ b/Engine/Source/Runtime/Interchange/Core/Public/Nodes/InterchangeBaseNodeContainer.h
@@ -0,0 +1,97 @@
+// Copyright Epic Games, Inc. All Rights Reserved.
+
+#pragma once
+
+#include "CoreMinimal.h"
+#include "Misc/FileHelper.h"
+#include "Nodes/InterchangeBaseNode.h"
+#include "Templates/UniquePtr.h"
+#include "UObject/Class.h"
+#include "UObject/Object.h"
+#include "UObject/ObjectMacros.h"
+
+#include "InterchangeBaseNodeContainer.generated.h"
+
+
+/**
+ * Interchange UInterchangeBaseNode graph is a format used to feed asset/scene import/reimport/export factories/writer.
+ * This container hold a flat list of all nodes that have been translated from the source data.
+ *
+ * Translators are filling this container and the Import/Export managers are reading it to execute the import/export process
+ */
+ UCLASS(BlueprintType)
+class INTERCHANGECORE_API UInterchangeBaseNodeContainer : public UObject
+{
+	GENERATED_BODY()
+public:
+
+	/**
+	 * Add a node in the container, the node will be add into a TMap.
+	 *
+	 * @param Node - a pointer on the node you want to add
+	 * @return: return the node unique ID of the added item. If the node already exist it will return the existing ID. Return InvalidNodeUID if the node cannot be added.
+	 *
+	 */
+	UFUNCTION(BlueprintCallable, Category = "Interchange | Node Container")
+	FName AddNode(UInterchangeBaseNode* Node);
+
+	/** Return true if the node unique ID exist in the container */
+	UFUNCTION(BlueprintCallable, Category = "Interchange | Node Container")
+	bool IsNodeUIDValid(FName NodeUniqueID) const;
+
+	/** Unordered iteration of the all nodes */
+	void IterateNodes(TFunctionRef<void(FName, const UInterchangeBaseNode*)> IterationLambda);
+
+	/** Return all nodes that do not have any parent */
+	UFUNCTION(BlueprintCallable, Category = "Interchange | Node Container")
+	void GetRoots(TArray<FName>& RootNodes);
+
+	/** Get an node pointer */
+	UFUNCTION(BlueprintCallable, Category = "Interchange | Node Container")
+	UInterchangeBaseNode* GetNode(FName NodeUniqueID);
+
+	/** Get an node pointer */
+	const UInterchangeBaseNode* GetNode(FName NodeUniqueID) const;
+
+	/** Set node ParentUID */
+	UFUNCTION(BlueprintCallable, Category = "Interchange | Node Container")
+	bool SetNodeParentUID(FName NodeUniqueID, FName NewParentNodeUID);
+
+	/** Get the node children count */
+	UFUNCTION(BlueprintCallable, Category = "Interchange | Node Container")
+	int32 GetNodeChildrenCount(FName NodeUniqueID) const;
+
+	/** Get all children UID */
+	UFUNCTION(BlueprintCallable, Category = "Interchange | Node Container")
+	TArray<FName> GetNodeChildrenUIDs(FName NodeUniqueID) const;
+
+	/** Get the node nth const children */
+	UFUNCTION(BlueprintCallable, Category = "Interchange | Node Container")
+	UInterchangeBaseNode* GetNodeChildren(FName NodeUniqueID, int32 ChildIndex);
+
+	/** Get the node nth const children. Const version */
+	const UInterchangeBaseNode* GetNodeChildren(FName NodeUniqueID, int32 ChildIndex) const;
+
+	/**
+	 * This function serialize the node container and all node sub-objects point by it.
+	 * Out of process translator like fbx will dump a file containing this data and the editor
+	 * will read the file and regenerate the container from the saved data.
+	 */
+	void SerializeNodeContainerData(FArchive& Ar);
+
+	/* Serialize the node container into the specified file.*/
+	UFUNCTION(BlueprintCallable, Category = "Interchange | Node Container")
+	void SaveToFile(const FString& Filename);
+
+	/* Serialize the node container from the specified file.*/
+	UFUNCTION(BlueprintCallable, Category = "Interchange | Node Container")
+	void LoadFromFile(const FString& Filename);
+
+private:
+
+	UInterchangeBaseNode* GetNodeChildrenInternal(FName NodeUniqueID, int32 ChildIndex);
+
+	/** Flat List of the nodes. Since the nodes are variable size, we store a pointer. */
+	UPROPERTY()
+	TMap<FName, UInterchangeBaseNode* > Nodes;
+};
diff --git a/Engine/Source/Runtime/Interchange/Core/Public/Nodes/InterchangeBaseNodeContainerAdapter.h b/Engine/Source/Runtime/Interchange/Core/Public/Nodes/InterchangeBaseNodeContainerAdapter.h
deleted file mode 100644
diff --git a/Engine/Source/Runtime/Interchange/Core/Public/Types/AttributeStorage.h b/Engine/Source/Runtime/Interchange/Core/Public/Types/AttributeStorage.h
--- a/Engine/Source/Runtime/Interchange/Core/Public/Types/AttributeStorage.h
+++ b/Engine/Source/Runtime/Interchange/Core/Public/Types/AttributeStorage.h
@@ -412,6 +412,10 @@ namespace Interchange
 			friend FAttributeStorage;
 		};
 
+		FAttributeStorage() = default;
+
+		FAttributeStorage(const FAttributeStorage& Other);
+
 		FAttributeStorage& operator=(const FAttributeStorage& Other);
 
 		/**
