Make sure interchange do not import fbx file over the content browser.

diff --git a/Engine/Plugins/Runtime/Interchange/Source/Import/Private/InterchangeImportPlugin.cpp b/Engine/Plugins/Runtime/Interchange/Source/Import/Private/InterchangeImportPlugin.cpp
--- a/Engine/Plugins/Runtime/Interchange/Source/Import/Private/InterchangeImportPlugin.cpp
+++ b/Engine/Plugins/Runtime/Interchange/Source/Import/Private/InterchangeImportPlugin.cpp
@@ -35,7 +35,7 @@ void FInterchangeImportPlugin::StartupModule()
 
 	//Register the translators
 	//Scenes
-	InterchangeManager.RegisterTranslator(UInterchangeFbxTranslator::StaticClass()); //Do not submit uncommented until we replace completly fbx importer (staticmesh + skeletalMesh + animation)
+	//InterchangeManager.RegisterTranslator(UInterchangeFbxTranslator::StaticClass()); //Do not submit uncommented until we replace completly fbx importer (staticmesh + skeletalMesh + animation)
 	//Textures
 	InterchangeManager.RegisterTranslator(UInterchangeBMPTranslator::StaticClass());
 	InterchangeManager.RegisterTranslator(UInterchangeDDSTranslator::StaticClass());
