fix CIS build, missing include file

diff --git a/Engine/Plugins/Runtime/Interchange/Source/Import/Public/Material/InterchangeMaterialPayloadInterface.h b/Engine/Plugins/Runtime/Interchange/Source/Import/Public/Material/InterchangeMaterialPayloadInterface.h
--- a/Engine/Plugins/Runtime/Interchange/Source/Import/Public/Material/InterchangeMaterialPayloadInterface.h
+++ b/Engine/Plugins/Runtime/Interchange/Source/Import/Public/Material/InterchangeMaterialPayloadInterface.h
@@ -2,9 +2,10 @@
 #pragma once
 
 #include "CoreMinimal.h"
+#include "InterchangeSourceData.h"
+#include "Material/InterchangeMaterialPayloadData.h"
 #include "UObject/Object.h"
 #include "UObject/ObjectMacros.h"
-#include "Material/InterchangeMaterialPayloadData.h"
 
 #include "InterchangeMaterialPayloadInterface.generated.h"
 
diff --git a/Engine/Plugins/Runtime/Interchange/Source/Import/Public/Mesh/InterchangeSkeletalMeshPayloadInterface.h b/Engine/Plugins/Runtime/Interchange/Source/Import/Public/Mesh/InterchangeSkeletalMeshPayloadInterface.h
--- a/Engine/Plugins/Runtime/Interchange/Source/Import/Public/Mesh/InterchangeSkeletalMeshPayloadInterface.h
+++ b/Engine/Plugins/Runtime/Interchange/Source/Import/Public/Mesh/InterchangeSkeletalMeshPayloadInterface.h
@@ -2,6 +2,7 @@
 #pragma once
 
 #include "CoreMinimal.h"
+#include "InterchangeSourceData.h"
 #include "Mesh/InterchangeSkeletalMeshPayload.h"
 #include "UObject/Object.h"
 #include "UObject/ObjectMacros.h"
diff --git a/Engine/Plugins/Runtime/Interchange/Source/Import/Public/Mesh/InterchangeStaticMeshPayloadInterface.h b/Engine/Plugins/Runtime/Interchange/Source/Import/Public/Mesh/InterchangeStaticMeshPayloadInterface.h
--- a/Engine/Plugins/Runtime/Interchange/Source/Import/Public/Mesh/InterchangeStaticMeshPayloadInterface.h
+++ b/Engine/Plugins/Runtime/Interchange/Source/Import/Public/Mesh/InterchangeStaticMeshPayloadInterface.h
@@ -2,9 +2,10 @@
 #pragma once
 
 #include "CoreMinimal.h"
+#include "InterchangeSourceData.h"
+#include "Mesh/InterchangeStaticMeshPayload.h"
 #include "UObject/Object.h"
 #include "UObject/ObjectMacros.h"
-#include "Mesh/InterchangeStaticMeshPayload.h"
 
 #include "InterchangeStaticMeshPayloadInterface.generated.h"
 
diff --git a/Engine/Plugins/Runtime/Interchange/Source/Import/Public/Texture/InterchangeTexturePayloadInterface.h b/Engine/Plugins/Runtime/Interchange/Source/Import/Public/Texture/InterchangeTexturePayloadInterface.h
--- a/Engine/Plugins/Runtime/Interchange/Source/Import/Public/Texture/InterchangeTexturePayloadInterface.h
+++ b/Engine/Plugins/Runtime/Interchange/Source/Import/Public/Texture/InterchangeTexturePayloadInterface.h
@@ -2,9 +2,10 @@
 #pragma once
 
 #include "CoreMinimal.h"
+#include "InterchangeSourceData.h"
+#include "Texture/InterchangeTexturePayloadData.h"
 #include "UObject/Object.h"
 #include "UObject/ObjectMacros.h"
-#include "Texture/InterchangeTexturePayloadData.h"
 
 #include "InterchangeTexturePayloadInterface.generated.h"
 
