Make sure all interchange namespace respect the coding standard

diff --git a/Engine/Plugins/Enterprise/DatasmithImporter/Source/DatasmithImporter/Private/DatasmithConsumer.cpp b/Engine/Plugins/Enterprise/DatasmithImporter/Source/DatasmithImporter/Private/DatasmithConsumer.cpp
--- a/Engine/Plugins/Enterprise/DatasmithImporter/Source/DatasmithImporter/Private/DatasmithConsumer.cpp
+++ b/Engine/Plugins/Enterprise/DatasmithImporter/Source/DatasmithImporter/Private/DatasmithConsumer.cpp
@@ -1529,7 +1529,7 @@ namespace DatasmithConsumerUtils
 					TextureElement->SetLabel( *Texture->GetName() );
 
 					TPromise< UObject* > TexturePromise = MakeFulfilledPromise< UObject* >( Texture );
-					ImportContext.ImportedTextures.FindOrAdd( TextureElement ) = Interchange::FAsyncImportResult{ TexturePromise.GetFuture(), FGraphEventRef() };
+					ImportContext.ImportedTextures.FindOrAdd( TextureElement ) = UE::Interchange::FAsyncImportResult{ TexturePromise.GetFuture(), FGraphEventRef() };
 					ImportContext.Scene->AddTexture( TextureElement );
 				}
 				else if(UMaterialInstance* MaterialInstance = Cast<UMaterialInstance>(Asset))
diff --git a/Engine/Plugins/Enterprise/DatasmithImporter/Source/DatasmithImporter/Private/DatasmithFileProducer.cpp b/Engine/Plugins/Enterprise/DatasmithImporter/Source/DatasmithImporter/Private/DatasmithFileProducer.cpp
--- a/Engine/Plugins/Enterprise/DatasmithImporter/Source/DatasmithImporter/Private/DatasmithFileProducer.cpp
+++ b/Engine/Plugins/Enterprise/DatasmithImporter/Source/DatasmithImporter/Private/DatasmithFileProducer.cpp
@@ -293,7 +293,7 @@ void UDatasmithFileProducer::SceneElementToWorld()
 		}
 	}
 
-	for ( TPair< TSharedRef< IDatasmithTextureElement >, Interchange::FAsyncImportResult >& AssetPair : ImportContextPtr->ImportedTextures )
+	for ( TPair< TSharedRef< IDatasmithTextureElement >, UE::Interchange::FAsyncImportResult >& AssetPair : ImportContextPtr->ImportedTextures )
 	{
 		if(AssetPair.Value.IsValid())
 		{
diff --git a/Engine/Plugins/Enterprise/DatasmithImporter/Source/DatasmithImporter/Private/DatasmithImporter.cpp b/Engine/Plugins/Enterprise/DatasmithImporter/Source/DatasmithImporter/Private/DatasmithImporter.cpp
--- a/Engine/Plugins/Enterprise/DatasmithImporter/Source/DatasmithImporter/Private/DatasmithImporter.cpp
+++ b/Engine/Plugins/Enterprise/DatasmithImporter/Source/DatasmithImporter/Private/DatasmithImporter.cpp
@@ -341,7 +341,7 @@ void FDatasmithImporter::ImportTextures( FDatasmithImportContext& ImportContext
 				break;
 			}
 
-			Interchange::FAsyncImportResult FutureTexture = DatasmithTextureImporter.CreateTextureAsync( FilteredTextureElements[TextureIndex] );
+			UE::Interchange::FAsyncImportResult FutureTexture = DatasmithTextureImporter.CreateTextureAsync( FilteredTextureElements[TextureIndex] );
 
 			if ( FutureTexture.IsValid() )
 			{
@@ -354,7 +354,7 @@ void FDatasmithImporter::ImportTextures( FDatasmithImportContext& ImportContext
 					}
 				);
 
-				Interchange::FAsyncImportResult& ImportedTexture = ImportContext.ImportedTextures.FindOrAdd( FilteredTextureElements[TextureIndex].ToSharedRef() );
+				UE::Interchange::FAsyncImportResult& ImportedTexture = ImportContext.ImportedTextures.FindOrAdd( FilteredTextureElements[TextureIndex].ToSharedRef() );
 				ImportedTexture = MoveTemp( FutureTexture );
 			}
 		}
@@ -443,9 +443,9 @@ UTexture* FDatasmithImporter::ImportTexture( FDatasmithImportContext& ImportCont
 {
 	TRACE_CPUPROFILER_EVENT_SCOPE(FDatasmithImporter::ImportTexture);
 
-	Interchange::FAsyncImportResult& ImportedTexture = ImportContext.ImportedTextures.FindOrAdd( TextureElement );
+	UE::Interchange::FAsyncImportResult& ImportedTexture = ImportContext.ImportedTextures.FindOrAdd( TextureElement );
 	TPromise< UObject* > TexturePromise = MakeFulfilledPromise< UObject* >( DatasmithTextureImporter.CreateTexture( TextureElement, TextureData, Extension ) );
-	ImportedTexture = Interchange::FAsyncImportResult{ TexturePromise.GetFuture(), FGraphEventRef() };
+	ImportedTexture = UE::Interchange::FAsyncImportResult{ TexturePromise.GetFuture(), FGraphEventRef() };
 
 	if ( !ImportedTexture.Get() )
 	{
@@ -1498,7 +1498,7 @@ void FDatasmithImporter::FinalizeImport(FDatasmithImportContext& ImportContext,
 	FScopedSlowTask* Progress = ProgressPtr.Get();
 
 	// Needs to be done in dependencies order (textures -> materials -> static meshes)
-	for (const TPair< TSharedRef< IDatasmithTextureElement >, Interchange::FAsyncImportResult >& ImportedTexturePair : ImportContext.ImportedTextures)
+	for (const TPair< TSharedRef< IDatasmithTextureElement >, UE::Interchange::FAsyncImportResult >& ImportedTexturePair : ImportContext.ImportedTextures)
 	{
 		if (ImportContext.bUserCancelled)
 		{
diff --git a/Engine/Plugins/Enterprise/DatasmithImporter/Source/DatasmithImporter/Private/DatasmithTextureImporter.cpp b/Engine/Plugins/Enterprise/DatasmithImporter/Source/DatasmithImporter/Private/DatasmithTextureImporter.cpp
--- a/Engine/Plugins/Enterprise/DatasmithImporter/Source/DatasmithImporter/Private/DatasmithTextureImporter.cpp
+++ b/Engine/Plugins/Enterprise/DatasmithImporter/Source/DatasmithImporter/Private/DatasmithTextureImporter.cpp
@@ -308,9 +308,9 @@ UTexture* FDatasmithTextureImporter::CreateTexture(const TSharedPtr<IDatasmithTe
 	return Texture;
 }
 
-Interchange::FAsyncImportResult FDatasmithTextureImporter::CreateTextureAsync(const TSharedPtr<IDatasmithTextureElement>& TextureElement)
+UE::Interchange::FAsyncImportResult FDatasmithTextureImporter::CreateTextureAsync(const TSharedPtr<IDatasmithTextureElement>& TextureElement)
 {
-	Interchange::FScopedSourceData ScopedSourceData( TextureElement->GetFile() );
+	UE::Interchange::FScopedSourceData ScopedSourceData( TextureElement->GetFile() );
 
 	const FString ContentPath = ImportContext.AssetsContext.TexturesImportPackage->GetPathName();
 
@@ -322,7 +322,7 @@ Interchange::FAsyncImportResult FDatasmithTextureImporter::CreateTextureAsync(co
 
 	ImportAssetParameters.OverridePipeline = TexturePipeline;
 
-	Interchange::FAsyncImportResult FutureTexture = UInterchangeManager::GetInterchangeManager().ImportAssetAsync( ContentPath, ScopedSourceData.GetSourceData(), ImportAssetParameters );
+	UE::Interchange::FAsyncImportResult FutureTexture = UInterchangeManager::GetInterchangeManager().ImportAssetAsync( ContentPath, ScopedSourceData.GetSourceData(), ImportAssetParameters );
 
 	return FutureTexture;
 }
diff --git a/Engine/Plugins/Enterprise/DatasmithImporter/Source/DatasmithImporter/Public/DatasmithImportContext.h b/Engine/Plugins/Enterprise/DatasmithImporter/Source/DatasmithImporter/Public/DatasmithImportContext.h
--- a/Engine/Plugins/Enterprise/DatasmithImporter/Source/DatasmithImporter/Public/DatasmithImportContext.h
+++ b/Engine/Plugins/Enterprise/DatasmithImporter/Source/DatasmithImporter/Public/DatasmithImportContext.h
@@ -256,7 +256,7 @@ struct DATASMITHIMPORTER_API FDatasmithImportContext
 	TMap< FString, TSharedRef < IDatasmithMeshElement > > ImportedStaticMeshesByName;
 
 	/** Map of imported texture for each texture element */
-	TMap< TSharedRef< IDatasmithTextureElement >, Interchange::FAsyncImportResult > ImportedTextures;
+	TMap< TSharedRef< IDatasmithTextureElement >, UE::Interchange::FAsyncImportResult > ImportedTextures;
 
 	/** Map of imported material for each material element */
 	TMap< TSharedRef< IDatasmithBaseMaterialElement >, UMaterialInterface* > ImportedMaterials;
diff --git a/Engine/Plugins/Enterprise/DatasmithImporter/Source/DatasmithImporter/Public/DatasmithTextureImporter.h b/Engine/Plugins/Enterprise/DatasmithImporter/Source/DatasmithImporter/Public/DatasmithTextureImporter.h
--- a/Engine/Plugins/Enterprise/DatasmithImporter/Source/DatasmithImporter/Public/DatasmithTextureImporter.h
+++ b/Engine/Plugins/Enterprise/DatasmithImporter/Source/DatasmithImporter/Public/DatasmithTextureImporter.h
@@ -17,9 +17,12 @@ class IDatasmithTextureElement;
 class UTexture;
 class UTextureFactory;
 
-namespace Interchange
+namespace UE
 {
-	class FAsyncImportResult;
+	namespace Interchange
+	{
+		class FAsyncImportResult;
+	}
 }
 
 class FDatasmithTextureImporter : private FNoncopyable
@@ -31,7 +34,7 @@ public:
 	UTexture* CreateTexture(const TSharedPtr<IDatasmithTextureElement>& TextureElement, const TArray<uint8>& TextureData, const FString& Extension);
 	bool GetTextureData(const TSharedPtr<IDatasmithTextureElement>& TextureElement, TArray<uint8>& TextureData, FString& Extension);
 
-	Interchange::FAsyncImportResult CreateTextureAsync(const TSharedPtr<IDatasmithTextureElement>& TextureElement);
+	UE::Interchange::FAsyncImportResult CreateTextureAsync(const TSharedPtr<IDatasmithTextureElement>& TextureElement);
 
 private:
 	bool ResizeTextureElement(const TSharedPtr<IDatasmithTextureElement>& TextureElement, FString& ResizedFilename);
diff --git a/Engine/Plugins/Enterprise/DatasmithImporter/Source/DatasmithImporter/Public/Utility/DatasmithImporterUtils.h b/Engine/Plugins/Enterprise/DatasmithImporter/Source/DatasmithImporter/Public/Utility/DatasmithImporterUtils.h
--- a/Engine/Plugins/Enterprise/DatasmithImporter/Source/DatasmithImporter/Public/Utility/DatasmithImporterUtils.h
+++ b/Engine/Plugins/Enterprise/DatasmithImporter/Source/DatasmithImporter/Public/Utility/DatasmithImporterUtils.h
@@ -250,7 +250,7 @@ struct FDatasmithFindAssetTypeHelper< UStaticMesh >
 template<>
 struct FDatasmithFindAssetTypeHelper< UTexture >
 {
-	static const TMap< TSharedRef< IDatasmithTextureElement >, Interchange::FAsyncImportResult >& GetImportedAssetsMap( const FDatasmithAssetsImportContext& AssetsContext )
+	static const TMap< TSharedRef< IDatasmithTextureElement >, UE::Interchange::FAsyncImportResult >& GetImportedAssetsMap( const FDatasmithAssetsImportContext& AssetsContext )
 	{
 		return AssetsContext.ParentContext.ImportedTextures;
 	}
@@ -421,11 +421,11 @@ inline UTexture* FDatasmithImporterUtils::FindAsset< UTexture >( const FDatasmit
 	if ( FPaths::IsRelative( ObjectPathName ) )
 	{
 		const TSharedRef< IDatasmithTextureElement >* ImportedElement   = FDatasmithFindAssetTypeHelper< UTexture >::GetImportedElementByName( AssetsContext, ObjectPathName );
-		const TMap< TSharedRef< IDatasmithTextureElement >, Interchange::FAsyncImportResult >& ImportedAssetsMap = FDatasmithFindAssetTypeHelper< UTexture >::GetImportedAssetsMap( AssetsContext );
+		const TMap< TSharedRef< IDatasmithTextureElement >, UE::Interchange::FAsyncImportResult >& ImportedAssetsMap = FDatasmithFindAssetTypeHelper< UTexture >::GetImportedAssetsMap( AssetsContext );
 
 		if (ImportedElement)
 		{
-			Interchange::FAsyncImportResult* Result =  const_cast< Interchange::FAsyncImportResult* >( ImportedAssetsMap.Find(*ImportedElement) );
+			UE::Interchange::FAsyncImportResult* Result =  const_cast<UE::Interchange::FAsyncImportResult* >( ImportedAssetsMap.Find(*ImportedElement) );
 			ensure(Result);
 			if (Result)
 			{
@@ -434,11 +434,11 @@ inline UTexture* FDatasmithImporterUtils::FindAsset< UTexture >( const FDatasmit
 		}
 		else
 		{
-			for ( const TPair< TSharedRef< IDatasmithTextureElement >, Interchange::FAsyncImportResult >& ImportedAssetPair : ImportedAssetsMap )
+			for ( const TPair< TSharedRef< IDatasmithTextureElement >, UE::Interchange::FAsyncImportResult >& ImportedAssetPair : ImportedAssetsMap )
 			{
 				if ( FCString::Stricmp( ImportedAssetPair.Key->GetName(), ObjectPathName ) == 0 )
 				{
-					Interchange::FAsyncImportResult& Result = const_cast< Interchange::FAsyncImportResult& >( ImportedAssetPair.Value );
+					UE::Interchange::FAsyncImportResult& Result = const_cast<UE::Interchange::FAsyncImportResult& >( ImportedAssetPair.Value );
 					return Cast< UTexture >( Result.Get() );
 				}
 			}
diff --git a/Engine/Plugins/Runtime/Interchange/Source/Dispatcher/Private/InterchangeCommands.cpp b/Engine/Plugins/Runtime/Interchange/Source/Dispatcher/Private/InterchangeCommands.cpp
--- a/Engine/Plugins/Runtime/Interchange/Source/Dispatcher/Private/InterchangeCommands.cpp
+++ b/Engine/Plugins/Runtime/Interchange/Source/Dispatcher/Private/InterchangeCommands.cpp
@@ -6,52 +6,55 @@
 #include "Serialization/MemoryWriter.h"
 #include "Serialization/MemoryReader.h"
 
-namespace InterchangeDispatcher
+namespace UE
 {
-	TSharedPtr<ICommand> CreateCommand(ECommandId CommandType)
+	namespace Interchange
 	{
-		switch (CommandType)
+		TSharedPtr<ICommand> CreateCommand(ECommandId CommandType)
 		{
-			case ECommandId::Ping: return MakeShared<FPingCommand>();
-			case ECommandId::BackPing: return MakeShared<FBackPingCommand>();
-			case ECommandId::RunTask: return MakeShared<FRunTaskCommand>();
-			case ECommandId::NotifyEndTask: return MakeShared<FCompletedTaskCommand>();
-			case ECommandId::Terminate: return MakeShared<FTerminateCommand>();
+			switch (CommandType)
+			{
+				case ECommandId::Ping: return MakeShared<FPingCommand>();
+				case ECommandId::BackPing: return MakeShared<FBackPingCommand>();
+				case ECommandId::RunTask: return MakeShared<FRunTaskCommand>();
+				case ECommandId::NotifyEndTask: return MakeShared<FCompletedTaskCommand>();
+				case ECommandId::Terminate: return MakeShared<FTerminateCommand>();
+			}
+			return nullptr;
 		}
-		return nullptr;
-	}
 
-	void SerializeCommand(ICommand& Command, TArray<uint8>& OutBuffer)
-	{
-		FMemoryWriter ArWriter(OutBuffer);
-		uint8 type = static_cast<uint8>(Command.GetType());
-		ArWriter << type;
-		ArWriter << Command;
-	}
+		void SerializeCommand(ICommand& Command, TArray<uint8>& OutBuffer)
+		{
+			FMemoryWriter ArWriter(OutBuffer);
+			uint8 type = static_cast<uint8>(Command.GetType());
+			ArWriter << type;
+			ArWriter << Command;
+		}
 
-	TSharedPtr<ICommand> DeserializeCommand(const TArray<uint8>& InBuffer)
-	{
-		FMemoryReader ArReader(InBuffer);
-		uint8 type = 0;
-		ArReader << type;
-		if (TSharedPtr<ICommand> Command = CreateCommand(static_cast<ECommandId>(type)))
+		TSharedPtr<ICommand> DeserializeCommand(const TArray<uint8>& InBuffer)
 		{
-			ArReader << *Command;
-			return ArReader.IsError() ? nullptr : Command;
+			FMemoryReader ArReader(InBuffer);
+			uint8 type = 0;
+			ArReader << type;
+			if (TSharedPtr<ICommand> Command = CreateCommand(static_cast<ECommandId>(type)))
+			{
+				ArReader << *Command;
+				return ArReader.IsError() ? nullptr : Command;
+			}
+			return nullptr;
 		}
-		return nullptr;
-	}
 
-	void FRunTaskCommand::SerializeImpl(FArchive& Ar)
-	{
-		Ar << JsonDescription;
-		Ar << JobIndex;
-	}
+		void FRunTaskCommand::SerializeImpl(FArchive& Ar)
+		{
+			Ar << JsonDescription;
+			Ar << JobIndex;
+		}
 
-  	void FCompletedTaskCommand::SerializeImpl(FArchive& Ar)
-	{
-		Ar << ProcessResult;
-		Ar << JSonMessages;
-		Ar << JSonResult;
-	}
-}
+		void FCompletedTaskCommand::SerializeImpl(FArchive& Ar)
+		{
+			Ar << ProcessResult;
+			Ar << JSonMessages;
+			Ar << JSonResult;
+		}
+	}//ns Interchange
+}//ns UE
diff --git a/Engine/Plugins/Runtime/Interchange/Source/Dispatcher/Private/InterchangeDispatcher.cpp b/Engine/Plugins/Runtime/Interchange/Source/Dispatcher/Private/InterchangeDispatcher.cpp
--- a/Engine/Plugins/Runtime/Interchange/Source/Dispatcher/Private/InterchangeDispatcher.cpp
+++ b/Engine/Plugins/Runtime/Interchange/Source/Dispatcher/Private/InterchangeDispatcher.cpp
@@ -10,165 +10,168 @@
 #include "Misc/Paths.h"
 #include "Misc/ScopeLock.h"
 
-namespace InterchangeDispatcher
+namespace UE
 {
-
-FInterchangeDispatcher::FInterchangeDispatcher(const FString& InResultFolder)
-	: NextTaskIndex(0)
-	, CompletedTaskCount(0)
-	, ResultFolder(InResultFolder)
-{
-}
-
-int32 FInterchangeDispatcher::AddTask(const FString& InJsonDescription)
-{
-	FScopeLock Lock(&TaskPoolCriticalSection);
-	int32 TaskIndex = TaskPool.Emplace(InJsonDescription);
-	TaskPool[TaskIndex].Index = TaskIndex;
-	return TaskIndex;
-}
-
-TOptional<FTask> FInterchangeDispatcher::GetNextTask()
-{
-	FScopeLock Lock(&TaskPoolCriticalSection);
-
-	while (TaskPool.IsValidIndex(NextTaskIndex) && TaskPool[NextTaskIndex].State != ETaskState::UnTreated)
-	{
-		NextTaskIndex++;
-	}
-
-	if (!TaskPool.IsValidIndex(NextTaskIndex))
-	{
-		return TOptional<FTask>();
-	}
-
-	TaskPool[NextTaskIndex].State = ETaskState::Running;
-	return TaskPool[NextTaskIndex++];
-}
-
-void FInterchangeDispatcher::SetTaskState(int32 TaskIndex, ETaskState TaskState, const FString& JsonResult, const TArray<FString>& JSonMessages)
-{
-	FString JsonDescription;
+	namespace Interchange
 	{
-		FScopeLock Lock(&TaskPoolCriticalSection);
 
-		if (!ensure(TaskPool.IsValidIndex(TaskIndex)))
+		FInterchangeDispatcher::FInterchangeDispatcher(const FString& InResultFolder)
+			: NextTaskIndex(0)
+			, CompletedTaskCount(0)
+			, ResultFolder(InResultFolder)
 		{
-			return;
 		}
 
-		FTask& Task = TaskPool[TaskIndex];
-		Task.State = TaskState;
-		Task.JsonResult = JsonResult;
-		Task.JsonMessages = JSonMessages;
-		JsonDescription = Task.JsonDescription;
-
-		if (TaskState == ETaskState::ProcessOk
-		 || TaskState == ETaskState::ProcessFailed)
+		int32 FInterchangeDispatcher::AddTask(const FString& InJsonDescription)
 		{
-			CompletedTaskCount++;
+			FScopeLock Lock(&TaskPoolCriticalSection);
+			int32 TaskIndex = TaskPool.Emplace(InJsonDescription);
+			TaskPool[TaskIndex].Index = TaskIndex;
+			return TaskIndex;
 		}
 
-		if (TaskState == ETaskState::UnTreated)
+		TOptional<FTask> FInterchangeDispatcher::GetNextTask()
 		{
-			NextTaskIndex = TaskIndex;
-		}
-	}
+			FScopeLock Lock(&TaskPoolCriticalSection);
 
-	UE_CLOG(TaskState == ETaskState::ProcessOk, LogInterchangeDispatcher, Verbose, TEXT("Json processed: %s"), *JsonDescription);
-	UE_CLOG(TaskState == ETaskState::UnTreated, LogInterchangeDispatcher, Warning, TEXT("Json resubmitted: %s"), *JsonDescription);
-	UE_CLOG(TaskState == ETaskState::ProcessFailed, LogInterchangeDispatcher, Error, TEXT("Json processing failure: %s"), *JsonDescription);
-}
+			while (TaskPool.IsValidIndex(NextTaskIndex) && TaskPool[NextTaskIndex].State != ETaskState::UnTreated)
+			{
+				NextTaskIndex++;
+			}
 
-void FInterchangeDispatcher::GetTaskState(int32 TaskIndex, ETaskState& TaskState, FString& JsonResult, TArray<FString>& JSonMessages)
-{
-	FScopeLock Lock(&TaskPoolCriticalSection);
+			if (!TaskPool.IsValidIndex(NextTaskIndex))
+			{
+				return TOptional<FTask>();
+			}
 
-	if (!ensure(TaskPool.IsValidIndex(TaskIndex)))
-	{
-		return;
-	}
+			TaskPool[NextTaskIndex].State = ETaskState::Running;
+			return TaskPool[NextTaskIndex++];
+		}
 
-	FTask& Task = TaskPool[TaskIndex];
-	TaskState = Task.State;
-	JsonResult = Task.JsonResult;
-	JSonMessages = Task.JsonMessages;
-}
+		void FInterchangeDispatcher::SetTaskState(int32 TaskIndex, ETaskState TaskState, const FString& JsonResult, const TArray<FString>& JSonMessages)
+		{
+			FString JsonDescription;
+			{
+				FScopeLock Lock(&TaskPoolCriticalSection);
+
+				if (!ensure(TaskPool.IsValidIndex(TaskIndex)))
+				{
+					return;
+				}
+
+				FTask& Task = TaskPool[TaskIndex];
+				Task.State = TaskState;
+				Task.JsonResult = JsonResult;
+				Task.JsonMessages = JSonMessages;
+				JsonDescription = Task.JsonDescription;
+
+				if (TaskState == ETaskState::ProcessOk
+					|| TaskState == ETaskState::ProcessFailed)
+				{
+					CompletedTaskCount++;
+				}
+
+				if (TaskState == ETaskState::UnTreated)
+				{
+					NextTaskIndex = TaskIndex;
+				}
+			}
+
+			UE_CLOG(TaskState == ETaskState::ProcessOk, LogInterchangeDispatcher, Verbose, TEXT("Json processed: %s"), *JsonDescription);
+			UE_CLOG(TaskState == ETaskState::UnTreated, LogInterchangeDispatcher, Warning, TEXT("Json resubmitted: %s"), *JsonDescription);
+			UE_CLOG(TaskState == ETaskState::ProcessFailed, LogInterchangeDispatcher, Error, TEXT("Json processing failure: %s"), *JsonDescription);
+		}
 
-void FInterchangeDispatcher::StartProcess()
-{
-	//Start the process
-	SpawnHandler();
-}
+		void FInterchangeDispatcher::GetTaskState(int32 TaskIndex, ETaskState& TaskState, FString& JsonResult, TArray<FString>& JSonMessages)
+		{
+			FScopeLock Lock(&TaskPoolCriticalSection);
 
-void FInterchangeDispatcher::TerminateProcess()
-{
-	//Empty the cache folder
-	if (IFileManager::Get().DirectoryExists(*ResultFolder))
-	{
-		const bool RequireExists = false;
-		//Delete recursively folder's content
-		const bool Tree = true;
-		IFileManager::Get().DeleteDirectory(*ResultFolder, RequireExists, Tree);
-	}
-	//Terminate the process
-	CloseHandler();
-}
-
-void FInterchangeDispatcher::WaitAllTaskToCompleteExecution()
-{
-	if (WorkerHandler.IsValid())
-	{
-		UE_LOG(LogInterchangeDispatcher, Error, TEXT("Cannot execute tasks before starting the process"));
-	}
+			if (!ensure(TaskPool.IsValidIndex(TaskIndex)))
+			{
+				return;
+			}
 
-	bool bLogRestartError = true;
-	while (!IsOver())
-	{
+			FTask& Task = TaskPool[TaskIndex];
+			TaskState = Task.State;
+			JsonResult = Task.JsonResult;
+			JSonMessages = Task.JsonMessages;
+		}
 
-		if (!IsHandlerAlive())
+		void FInterchangeDispatcher::StartProcess()
 		{
-			break;
+			//Start the process
+			SpawnHandler();
 		}
 
-		FPlatformProcess::Sleep(0.1f);
-	}
+		void FInterchangeDispatcher::TerminateProcess()
+		{
+			//Empty the cache folder
+			if (IFileManager::Get().DirectoryExists(*ResultFolder))
+			{
+				const bool RequireExists = false;
+				//Delete recursively folder's content
+				const bool Tree = true;
+				IFileManager::Get().DeleteDirectory(*ResultFolder, RequireExists, Tree);
+			}
+			//Terminate the process
+			CloseHandler();
+		}
 
-	if (!IsOver())
-	{
-		UE_LOG(LogInterchangeDispatcher, Warning,
-			TEXT("Begin local processing. (Multi Process failed to consume all the tasks)\n")
-			TEXT("See workers logs: %sPrograms/InterchangeWorker/Saved/Logs"), *FPaths::ConvertRelativePathToFull(FPaths::EngineDir()));
-	}
-	else
-	{
-		UE_LOG(LogInterchangeDispatcher, Display, TEXT("Multi Process ended and consumed all the tasks"));
-	}
-}
+		void FInterchangeDispatcher::WaitAllTaskToCompleteExecution()
+		{
+			if (WorkerHandler.IsValid())
+			{
+				UE_LOG(LogInterchangeDispatcher, Error, TEXT("Cannot execute tasks before starting the process"));
+			}
+
+			bool bLogRestartError = true;
+			while (!IsOver())
+			{
+
+				if (!IsHandlerAlive())
+				{
+					break;
+				}
+
+				FPlatformProcess::Sleep(0.1f);
+			}
+
+			if (!IsOver())
+			{
+				UE_LOG(LogInterchangeDispatcher, Warning,
+					   TEXT("Begin local processing. (Multi Process failed to consume all the tasks)\n")
+					   TEXT("See workers logs: %sPrograms/InterchangeWorker/Saved/Logs"), *FPaths::ConvertRelativePathToFull(FPaths::EngineDir()));
+			}
+			else
+			{
+				UE_LOG(LogInterchangeDispatcher, Display, TEXT("Multi Process ended and consumed all the tasks"));
+			}
+		}
 
-bool FInterchangeDispatcher::IsOver()
-{
-	FScopeLock Lock(&TaskPoolCriticalSection);
-	return CompletedTaskCount == TaskPool.Num();
-}
+		bool FInterchangeDispatcher::IsOver()
+		{
+			FScopeLock Lock(&TaskPoolCriticalSection);
+			return CompletedTaskCount == TaskPool.Num();
+		}
 
-void FInterchangeDispatcher::SpawnHandler()
-{
-	WorkerHandler = MakeUnique<FInterchangeWorkerHandler>(*this, ResultFolder);
-}
+		void FInterchangeDispatcher::SpawnHandler()
+		{
+			WorkerHandler = MakeUnique<FInterchangeWorkerHandler>(*this, ResultFolder);
+		}
 
-bool FInterchangeDispatcher::IsHandlerAlive()
-{
-	return WorkerHandler.IsValid() && WorkerHandler->IsAlive();
-}
+		bool FInterchangeDispatcher::IsHandlerAlive()
+		{
+			return WorkerHandler.IsValid() && WorkerHandler->IsAlive();
+		}
 
-void FInterchangeDispatcher::CloseHandler()
-{
-	if (IsHandlerAlive())
-	{
-		WorkerHandler->Stop();
-	}
-	WorkerHandler.Reset();
-}
+		void FInterchangeDispatcher::CloseHandler()
+		{
+			if (IsHandlerAlive())
+			{
+				WorkerHandler->Stop();
+			}
+			WorkerHandler.Reset();
+		}
 
-} // ns InterchangeDispatcher
\ No newline at end of file
+	} //ns Interchange
+}//ns UE
\ No newline at end of file
diff --git a/Engine/Plugins/Runtime/Interchange/Source/Dispatcher/Private/InterchangeDispatcherNetworking.cpp b/Engine/Plugins/Runtime/Interchange/Source/Dispatcher/Private/InterchangeDispatcherNetworking.cpp
--- a/Engine/Plugins/Runtime/Interchange/Source/Dispatcher/Private/InterchangeDispatcherNetworking.cpp
+++ b/Engine/Plugins/Runtime/Interchange/Source/Dispatcher/Private/InterchangeDispatcherNetworking.cpp
@@ -10,361 +10,361 @@
 #include "Serialization/MemoryReader.h"
 #include "Serialization/MemoryWriter.h"
 
-
-namespace InterchangeDispatcher
-{
-
-static FTimespan TimespanFromSeconds(double Seconds)
-{
-	return FTimespan(int64(Seconds * ETimespan::TicksPerSecond));
-}
-
-FNetworkNode::~FNetworkNode()
-{
-	CloseSocket(ConnectedSocket);
-}
-
-bool FNetworkNode::SendMessage(const TArray<uint8>& Message, double Timeout_s)
+namespace UE
 {
-	if (bWriteError)
-	{
-		UE_LOG(LogInterchangeDispatcher, Error, TEXT("bWriteError flag raised, can't write"));
-		return false;
-	}
-
-	if (ConnectedSocket == nullptr)
-	{
-		UE_LOG(LogInterchangeDispatcher, Error, TEXT("node not connected, can't write"));
-		return false;
-	}
-
-	ESocketConnectionState State = ConnectedSocket->GetConnectionState();
-	if (State == ESocketConnectionState::SCS_ConnectionError)
-	{
-		UE_LOG(LogInterchangeDispatcher, Error, TEXT("connection state error"));
-		return false;
-	}
-
-	bool bCanWrite = ConnectedSocket->Wait(ESocketWaitConditions::WaitForWrite, TimespanFromSeconds(Timeout_s));
-	if (!bCanWrite)
+	namespace Interchange
 	{
-		UE_LOG(LogInterchangeDispatcher, Warning, TEXT("can't write on socket"));
-		return false;
-	}
 
-	ConnectedSocket->SetNonBlocking(false);
-
-	auto SendBuffer = [&](const TArray<uint8>& Buffer) -> bool
-	{
-		uint32 BufferSize = Buffer.Num();
-		const uint8* Data = Buffer.GetData();
+		static FTimespan TimespanFromSeconds(double Seconds)
+		{
+			return FTimespan(int64(Seconds * ETimespan::TicksPerSecond));
+		}
 
-		uint32 TotalByteSent = 0;
-		bool bSendSucceed = true;
-		while (bSendSucceed && TotalByteSent < BufferSize)
+		FNetworkNode::~FNetworkNode()
 		{
-			int32 BytesSent = 0;
-			bSendSucceed &= ConnectedSocket->Send(Data + TotalByteSent, BufferSize - TotalByteSent, BytesSent);
-			TotalByteSent += (uint32) BytesSent;
+			CloseSocket(ConnectedSocket);
 		}
-		return bSendSucceed && TotalByteSent == BufferSize;
-	};
 
-	// Send the header
-	TArray<uint8> HeaderBuffer;
-	FMemoryWriter ArWriter(HeaderBuffer);
+		bool FNetworkNode::SendMessage(const TArray<uint8>& Message, double Timeout_s)
+		{
+			if (bWriteError)
+			{
+				UE_LOG(LogInterchangeDispatcher, Error, TEXT("bWriteError flag raised, can't write"));
+				return false;
+			}
+
+			if (ConnectedSocket == nullptr)
+			{
+				UE_LOG(LogInterchangeDispatcher, Error, TEXT("node not connected, can't write"));
+				return false;
+			}
+
+			ESocketConnectionState State = ConnectedSocket->GetConnectionState();
+			if (State == ESocketConnectionState::SCS_ConnectionError)
+			{
+				UE_LOG(LogInterchangeDispatcher, Error, TEXT("connection state error"));
+				return false;
+			}
+
+			bool bCanWrite = ConnectedSocket->Wait(ESocketWaitConditions::WaitForWrite, TimespanFromSeconds(Timeout_s));
+			if (!bCanWrite)
+			{
+				UE_LOG(LogInterchangeDispatcher, Warning, TEXT("can't write on socket"));
+				return false;
+			}
+
+			ConnectedSocket->SetNonBlocking(false);
+
+			auto SendBuffer = [&](const TArray<uint8>& Buffer) -> bool
+			{
+				uint32 BufferSize = Buffer.Num();
+				const uint8* Data = Buffer.GetData();
+
+				uint32 TotalByteSent = 0;
+				bool bSendSucceed = true;
+				while (bSendSucceed && TotalByteSent < BufferSize)
+				{
+					int32 BytesSent = 0;
+					bSendSucceed &= ConnectedSocket->Send(Data + TotalByteSent, BufferSize - TotalByteSent, BytesSent);
+					TotalByteSent += (uint32)BytesSent;
+				}
+				return bSendSucceed && TotalByteSent == BufferSize;
+			};
+
+			// Send the header
+			TArray<uint8> HeaderBuffer;
+			FMemoryWriter ArWriter(HeaderBuffer);
 
-	FMessageHeader Header;
-	Header.ByteSize = Message.Num();
-	ArWriter << Header;
+			FMessageHeader Header;
+			Header.ByteSize = Message.Num();
+			ArWriter << Header;
 
-	if (!SendBuffer(HeaderBuffer))
-	{
-		UE_LOG(LogInterchangeDispatcher, Error, TEXT("can't write header on socket"));
-		bWriteError = true;
-		return false;
-	}
+			if (!SendBuffer(HeaderBuffer))
+			{
+				UE_LOG(LogInterchangeDispatcher, Error, TEXT("can't write header on socket"));
+				bWriteError = true;
+				return false;
+			}
+
+			// Send the content
+			if (!SendBuffer(Message))
+			{
+				UE_LOG(LogInterchangeDispatcher, Error, TEXT("can't write content on socket"));
+				bWriteError = true;
+				return false;
+			}
 
-	// Send the content
-	if (!SendBuffer(Message))
-	{
-		UE_LOG(LogInterchangeDispatcher, Error, TEXT("can't write content on socket"));
-		bWriteError = true;
-		return false;
-	}
+			return true;
+		}
 
-	return true;
-}
+		bool FNetworkNode::ReceiveMessage(TArray<uint8>& OutMessage, double Timeout_s)
+		{
+			if (bReadError)
+			{
+				UE_LOG(LogInterchangeDispatcher, Error, TEXT("ReadError flag raised, can't read"));
+				return false;
+			}
+
+			if (ConnectedSocket == nullptr)
+			{
+				UE_LOG(LogInterchangeDispatcher, Error, TEXT("node not connected, can't read"));
+				return false;
+			}
+
+			ConnectedSocket->SetNonBlocking(true);
+
+			bool bCanRead = ConnectedSocket && ConnectedSocket->Wait(ESocketWaitConditions::WaitForRead, TimespanFromSeconds(Timeout_s));
+			if (!bCanRead)
+			{
+				return false;
+			}
+
+			// Make sure we have a valid header
+			if (IncommingMessage.Header.ByteSize < 0)
+			{
+				static uint32 HeaderByteSize = []()
+				{
+					FMessageHeader Header;
+					TArray<uint8> HeaderBuffer;
+					FMemoryWriter ArWriter(HeaderBuffer);
+					ArWriter << Header;
+					return HeaderBuffer.Num();
+				}();
+
+				uint32 PendingByteCount;
+				if (!ConnectedSocket->HasPendingData(PendingByteCount) || PendingByteCount < HeaderByteSize)
+				{
+					return false;
+				}
+
+				TArray<uint8> HeaderBuffer;
+				HeaderBuffer.AddZeroed(HeaderByteSize);
+				int32 BytesRead = -1;
+				bool bRecvSucceed = ConnectedSocket->Recv(HeaderBuffer.GetData(), HeaderBuffer.Num(), BytesRead);
+				if (!bRecvSucceed || BytesRead != HeaderByteSize)
+				{
+					UE_LOG(LogInterchangeDispatcher, Error, TEXT("Parsed header failed"));
+					bReadError = true;
+					return false;
+				}
+
+				FMemoryReader ArReader(HeaderBuffer);
+				ArReader << IncommingMessage.Header;
+
+				if (IncommingMessage.Header.ByteSize < 0 || IncommingMessage.Header.ByteSize > 1 << 20)
+				{
+					UE_LOG(LogInterchangeDispatcher, Error, TEXT("Parsed header failed: bad message size %d"), IncommingMessage.Header.ByteSize);
+					bReadError = true;
+					return false;
+				}
+
+				IncommingMessage.Content.Reserve(IncommingMessage.Header.ByteSize);
+			}
+
+			// fill the message with available data
+			uint32 PendingByteCount;
+			if (IncommingMessage.Header.ByteSize >= 0 && ConnectedSocket->HasPendingData(PendingByteCount))
+			{
+				int32 MissingByteInCurrentMessage = IncommingMessage.Header.ByteSize - IncommingMessage.Content.Num();
+				int32 ReadTarget = int32(PendingByteCount) < MissingByteInCurrentMessage ? PendingByteCount : MissingByteInCurrentMessage;
+				uint8* Destination = IncommingMessage.Content.GetData() + IncommingMessage.Content.Num();
+				IncommingMessage.Content.AddUninitialized(ReadTarget);
+
+				int32 BytesRead = -1;
+				bool bRecvSucceed = ConnectedSocket->Recv(Destination, ReadTarget, BytesRead);
+				if (BytesRead != ReadTarget)
+				{
+					bReadError = true;
+					UE_LOG(LogInterchangeDispatcher, Error, TEXT("Recv issue"));
+					return false;
+				}
+
+				bool bMessageCompleted = BytesRead == MissingByteInCurrentMessage;
+				if (bMessageCompleted)
+				{
+					OutMessage = MoveTemp(IncommingMessage.Content);
+					IncommingMessage = FMessage();
+					return true;
+				}
+			}
 
-bool FNetworkNode::ReceiveMessage(TArray<uint8>& OutMessage, double Timeout_s)
-{
-	if (bReadError)
-	{
-		UE_LOG(LogInterchangeDispatcher, Error, TEXT("ReadError flag raised, can't read"));
-		return false;
-	}
+			return false;
+		}
 
-	if (ConnectedSocket == nullptr)
-	{
-		UE_LOG(LogInterchangeDispatcher, Error, TEXT("node not connected, can't read"));
-		return false;
-	}
 
-	ConnectedSocket->SetNonBlocking(true);
+		FSocket* FNetworkNode::CreateInternalSocket(const FString& Description)
+		{
+			FSocket* Socket = nullptr;
+			ISocketSubsystem* SocketSubsystem = ISocketSubsystem::Get();
+			TSharedRef<FInternetAddr> InternetAddress = SocketSubsystem->CreateInternetAddr();
+			InternetAddress->SetLoopbackAddress();
+
+			FName Protocol = NAME_None;
+			if (ensure(InternetAddress->IsValid()))
+			{
+				Protocol = InternetAddress->GetProtocolType();
+			}
+			FSocket* NewSocket = SocketSubsystem->CreateSocket(NAME_Stream, *Description, Protocol);
+			UE_CLOG(!NewSocket, LogInterchangeDispatcher, Error, TEXT("Socket creation failure"));
+			return NewSocket;
+		}
 
-	bool bCanRead = ConnectedSocket && ConnectedSocket->Wait(ESocketWaitConditions::WaitForRead, TimespanFromSeconds(Timeout_s));
-	if (!bCanRead)
-	{
-		return false;
-	}
+		void FNetworkNode::CloseSocket(FSocket*& Socket)
+		{
+			if (Socket)
+			{
+				Socket->Close();
+				ISocketSubsystem::Get()->DestroySocket(Socket);
+				Socket = nullptr;
+			}
+		}
 
-	// Make sure we have a valid header
-	if (IncommingMessage.Header.ByteSize < 0)
-	{
-		static uint32 HeaderByteSize = []()
+		bool FNetworkNode::IsConnected()
 		{
-			FMessageHeader Header;
-			TArray<uint8> HeaderBuffer;
-			FMemoryWriter ArWriter(HeaderBuffer);
-			ArWriter << Header;
-			return HeaderBuffer.Num();
-		}();
+			return ConnectedSocket && ConnectedSocket->GetConnectionState() == ESocketConnectionState::SCS_Connected;
+		}
 
-		uint32 PendingByteCount;
-		if (!ConnectedSocket->HasPendingData(PendingByteCount) || PendingByteCount < HeaderByteSize)
+		FNetworkServerNode::FNetworkServerNode()
 		{
-			return false;
+			ListeningSocket = CreateInternalSocket(TEXT("Interchange listening socket"));
+
+			if (!ListeningSocket)
+			{
+				ConnectedSocketError = SocketErrorCode::Error_Create;
+				return;
+			}
+
+			ISocketSubsystem* SocketSubsystem = ISocketSubsystem::Get();
+			TSharedRef<FInternetAddr> InternetAddress = SocketSubsystem->CreateInternetAddr();
+			InternetAddress->SetLoopbackAddress();
+
+			int32 BoundPort = SocketSubsystem->BindNextPort(ListeningSocket, *InternetAddress, 1e3, 1);
+			if (BoundPort == 0)
+			{
+				ConnectedSocketError = SocketErrorCode::Error_Bind;
+				UE_LOG(LogInterchangeDispatcher, Error, TEXT("Socket binding failure"));
+				return;
+			}
+			ensure(BoundPort == ListeningSocket->GetPortNo());
+
+			if (!ListeningSocket->Listen(0))
+			{
+				ConnectedSocketError = SocketErrorCode::Error_Listen;
+				UE_LOG(LogInterchangeDispatcher, Error, TEXT("Socket listen failure"));
+				return;
+			}
 		}
 
-		TArray<uint8> HeaderBuffer;
-		HeaderBuffer.AddZeroed(HeaderByteSize);
-		int32 BytesRead = -1;
-		bool bRecvSucceed = ConnectedSocket->Recv(HeaderBuffer.GetData(), HeaderBuffer.Num(), BytesRead);
-		if (!bRecvSucceed || BytesRead != HeaderByteSize)
+		FNetworkServerNode::~FNetworkServerNode()
 		{
-			UE_LOG(LogInterchangeDispatcher, Error, TEXT("Parsed header failed"));
-			bReadError = true;
-			return false;
+			CloseSocket(ListeningSocket);
 		}
 
-		FMemoryReader ArReader(HeaderBuffer);
-		ArReader << IncommingMessage.Header;
+		int32 FNetworkServerNode::GetListeningPort()
+		{
+			return ListeningSocket ? ListeningSocket->GetPortNo() : 0;
+		}
 
-		if (IncommingMessage.Header.ByteSize < 0 || IncommingMessage.Header.ByteSize > 1<<20)
+		bool FNetworkServerNode::Accept(const FString& Description, double Timeout_s)
 		{
-			UE_LOG(LogInterchangeDispatcher, Error, TEXT("Parsed header failed: bad message size %d"), IncommingMessage.Header.ByteSize);
-			bReadError = true;
+			CloseSocket(ConnectedSocket);
+
+			// wait until a connection occurs
+			ListeningSocket->SetNonBlocking();
+			double AcceptLimit_s = FPlatformTime::Seconds() + Timeout_s;
+			do
+			{
+				ConnectedSocket = ListeningSocket->Accept(Description);
+				if (ConnectedSocket)
+				{
+					return true;
+				}
+				FPlatformProcess::Sleep(0.1);
+			} while (FPlatformTime::Seconds() < AcceptLimit_s);
+
 			return false;
 		}
 
-		IncommingMessage.Content.Reserve(IncommingMessage.Header.ByteSize);
-	}
-
-	// fill the message with available data
-	uint32 PendingByteCount;
-	if (IncommingMessage.Header.ByteSize >= 0 && ConnectedSocket->HasPendingData(PendingByteCount))
-	{
-		int32 MissingByteInCurrentMessage = IncommingMessage.Header.ByteSize - IncommingMessage.Content.Num();
-		int32 ReadTarget = int32(PendingByteCount) < MissingByteInCurrentMessage ? PendingByteCount : MissingByteInCurrentMessage;
-		uint8* Destination = IncommingMessage.Content.GetData() + IncommingMessage.Content.Num();
-		IncommingMessage.Content.AddUninitialized(ReadTarget);
-
-		int32 BytesRead = -1;
-		bool bRecvSucceed = ConnectedSocket->Recv(Destination, ReadTarget, BytesRead);
-		if (BytesRead != ReadTarget)
+		bool FNetworkClientNode::Connect(const FString& Description, int32 ServerPort, double Timeout_s)
 		{
-			bReadError = true;
-			UE_LOG(LogInterchangeDispatcher, Error, TEXT("Recv issue"));
+			CloseSocket(ConnectedSocket);
+			ConnectedSocket = CreateInternalSocket(Description);
+
+			if (!ConnectedSocket)
+			{
+				ConnectedSocketError = SocketErrorCode::Error_Create;
+				return false;
+			}
+
+			ISocketSubsystem* SocketSubsystem = ISocketSubsystem::Get();
+			TSharedRef<FInternetAddr> InternetAddress = SocketSubsystem->CreateInternetAddr();
+			InternetAddress->SetLoopbackAddress();
+			InternetAddress->SetPort(ServerPort);
+
+			double ConnectTimeout_s = FPlatformTime::Seconds() + Timeout_s;
+			ConnectedSocket->SetNonBlocking(true);
+			do
+			{
+				if (ConnectedSocket->Connect(*InternetAddress))
+				{
+					UE_LOG(LogInterchangeDispatcher, Verbose, TEXT("Client Node is connected"));
+					return true;
+				}
+			} while (FPlatformTime::Seconds() < ConnectTimeout_s);
+
+			UE_LOG(LogInterchangeDispatcher, Error, TEXT("Client socket failed to connect"));
+			CloseSocket(ConnectedSocket);
 			return false;
 		}
 
-		bool bMessageCompleted = BytesRead == MissingByteInCurrentMessage;
-		if (bMessageCompleted)
+		void FCommandQueue::SetNetworkInterface(FNetworkNode* InNetworkInterface)
 		{
-			OutMessage = MoveTemp(IncommingMessage.Content);
-			IncommingMessage = FMessage();
-			return true;
+			NetworkInterface = InNetworkInterface;
 		}
-	}
 
-	return false;
-}
-
-
-FSocket* FNetworkNode::CreateInternalSocket(const FString& Description)
-{
-	FSocket* Socket = nullptr;
-	ISocketSubsystem* SocketSubsystem = ISocketSubsystem::Get();
-	TSharedRef<FInternetAddr> InternetAddress = SocketSubsystem->CreateInternetAddr();
-	InternetAddress->SetLoopbackAddress();
-
-	FName Protocol = NAME_None;
-	if (ensure(InternetAddress->IsValid()))
-	{
-		Protocol = InternetAddress->GetProtocolType();
-	}
-	FSocket* NewSocket = SocketSubsystem->CreateSocket(NAME_Stream, *Description, Protocol);
-	UE_CLOG(!NewSocket, LogInterchangeDispatcher, Error, TEXT("Socket creation failure"));
-	return NewSocket;
-}
-
-void FNetworkNode::CloseSocket(FSocket*& Socket)
-{
-	if (Socket)
-	{
-		Socket->Close();
-		ISocketSubsystem::Get()->DestroySocket(Socket);
-		Socket = nullptr;
-	}
-}
-
-bool FNetworkNode::IsConnected()
-{
-	return ConnectedSocket && ConnectedSocket->GetConnectionState() == ESocketConnectionState::SCS_Connected;
-}
-
-FNetworkServerNode::FNetworkServerNode()
-{
-	ListeningSocket = CreateInternalSocket(TEXT("Interchange listening socket"));
-
-	if (!ListeningSocket)
-	{
-		ConnectedSocketError = SocketErrorCode::Error_Create;
-		return;
-	}
-
-	ISocketSubsystem* SocketSubsystem = ISocketSubsystem::Get();
-	TSharedRef<FInternetAddr> InternetAddress = SocketSubsystem->CreateInternetAddr();
-	InternetAddress->SetLoopbackAddress();
-
-	int32 BoundPort = SocketSubsystem->BindNextPort(ListeningSocket, *InternetAddress, 1e3, 1);
-	if (BoundPort == 0)
-	{
-		ConnectedSocketError = SocketErrorCode::Error_Bind;
-		UE_LOG(LogInterchangeDispatcher, Error, TEXT("Socket binding failure"));
-		return;
-	}
-	ensure(BoundPort == ListeningSocket->GetPortNo());
-
-	if (!ListeningSocket->Listen(0))
-	{
-		ConnectedSocketError = SocketErrorCode::Error_Listen;
-		UE_LOG(LogInterchangeDispatcher, Error, TEXT("Socket listen failure"));
-		return;
-	}
-}
-
-FNetworkServerNode::~FNetworkServerNode()
-{
-	CloseSocket(ListeningSocket);
-}
-
-int32 FNetworkServerNode::GetListeningPort()
-{
-	return ListeningSocket ? ListeningSocket->GetPortNo() : 0;
-}
-
-bool FNetworkServerNode::Accept(const FString& Description, double Timeout_s)
-{
-	CloseSocket(ConnectedSocket);
-
-	// wait until a connection occurs
-	ListeningSocket->SetNonBlocking();
-	double AcceptLimit_s = FPlatformTime::Seconds() + Timeout_s;
-	do
-	{
-		ConnectedSocket = ListeningSocket->Accept(Description);
-		if (ConnectedSocket)
+		TSharedPtr<ICommand> FCommandQueue::GetNextCommand(double Timeout_s)
 		{
-			return true;
+			// consume all available commands
+			while (Poll(Timeout_s))
+			{
+				Timeout_s = 0;
+			}
+
+			TSharedPtr<ICommand> OutCommand;
+			return InCommands.Dequeue(OutCommand) ? OutCommand : nullptr;
 		}
-		FPlatformProcess::Sleep(0.1);
-	}
-	while(FPlatformTime::Seconds() < AcceptLimit_s);
 
-	return false;
-}
-
-bool FNetworkClientNode::Connect(const FString& Description, int32 ServerPort, double Timeout_s)
-{
-	CloseSocket(ConnectedSocket);
-	ConnectedSocket = CreateInternalSocket(Description);
-
-	if (!ConnectedSocket)
-	{
-		ConnectedSocketError = SocketErrorCode::Error_Create;
-		return false;
-	}
-
-	ISocketSubsystem* SocketSubsystem = ISocketSubsystem::Get();
-	TSharedRef<FInternetAddr> InternetAddress = SocketSubsystem->CreateInternetAddr();
-	InternetAddress->SetLoopbackAddress();
-	InternetAddress->SetPort(ServerPort);
-
-	double ConnectTimeout_s = FPlatformTime::Seconds() + Timeout_s;
-	ConnectedSocket->SetNonBlocking(true);
-	do
-	{
-		if (ConnectedSocket->Connect(*InternetAddress))
+		bool FCommandQueue::SendCommand(ICommand& Commmand, double Timeout_s)
 		{
-			UE_LOG(LogInterchangeDispatcher, Verbose, TEXT("Client Node is connected"));
-			return true;
+			TArray<uint8> CommandBuffer;
+			SerializeCommand(Commmand, CommandBuffer);
+			return NetworkInterface && NetworkInterface->SendMessage(CommandBuffer, Timeout_s);
 		}
-	}
-	while(FPlatformTime::Seconds() < ConnectTimeout_s);
-
-	UE_LOG(LogInterchangeDispatcher, Error, TEXT("Client socket failed to connect"));
-	CloseSocket(ConnectedSocket);
-	return false;
-}
-
-void FCommandQueue::SetNetworkInterface(FNetworkNode* InNetworkInterface)
-{
-	NetworkInterface = InNetworkInterface;
-}
 
-TSharedPtr<ICommand> FCommandQueue::GetNextCommand(double Timeout_s)
-{
-	// consume all available commands
-	while (Poll(Timeout_s))
-	{
-		Timeout_s = 0;
-	}
-
-	TSharedPtr<ICommand> OutCommand;
-	return InCommands.Dequeue(OutCommand) ? OutCommand : nullptr;
-}
-
-bool FCommandQueue::SendCommand(ICommand& Commmand, double Timeout_s)
-{
-	TArray<uint8> CommandBuffer;
-	SerializeCommand(Commmand, CommandBuffer);
-	return NetworkInterface && NetworkInterface->SendMessage(CommandBuffer, Timeout_s);
-}
-
-void FCommandQueue::Disconnect(double Timeout_s)
-{
-	// consume all available commands before closing the connection
-	while (Poll(Timeout_s))
-	{
-		Timeout_s = 0;
-	}
+		void FCommandQueue::Disconnect(double Timeout_s)
+		{
+			// consume all available commands before closing the connection
+			while (Poll(Timeout_s))
+			{
+				Timeout_s = 0;
+			}
 
-	NetworkInterface = nullptr;
-}
+			NetworkInterface = nullptr;
+		}
 
-bool FCommandQueue::Poll(double Timeout_s)
-{
-	TArray<uint8> CommandBuffer;
-	if (NetworkInterface && NetworkInterface->ReceiveMessage(CommandBuffer, Timeout_s))
-	{
-		if (TSharedPtr<ICommand> Command = DeserializeCommand(CommandBuffer))
+		bool FCommandQueue::Poll(double Timeout_s)
 		{
-			InCommands.Enqueue(Command);
-			return true;
+			TArray<uint8> CommandBuffer;
+			if (NetworkInterface && NetworkInterface->ReceiveMessage(CommandBuffer, Timeout_s))
+			{
+				if (TSharedPtr<ICommand> Command = DeserializeCommand(CommandBuffer))
+				{
+					InCommands.Enqueue(Command);
+					return true;
+				}
+			}
+			return false;
 		}
-	}
-	return false;
-}
 
-} // ns InterchangeDispatcher
+	} //ns Interchange
+}//ns UE
diff --git a/Engine/Plugins/Runtime/Interchange/Source/Dispatcher/Private/InterchangeDispatcherTask.cpp b/Engine/Plugins/Runtime/Interchange/Source/Dispatcher/Private/InterchangeDispatcherTask.cpp
--- a/Engine/Plugins/Runtime/Interchange/Source/Dispatcher/Private/InterchangeDispatcherTask.cpp
+++ b/Engine/Plugins/Runtime/Interchange/Source/Dispatcher/Private/InterchangeDispatcherTask.cpp
@@ -9,103 +9,106 @@
 #include "Serialization/JsonWriter.h"
 #include "Serialization/JsonSerializer.h"
 
-namespace InterchangeDispatcher
+namespace UE
 {
-	FString FJsonLoadSourceCmd::ToJson() const
+	namespace Interchange
 	{
-		//Code should not do a ToJson if the data was not set before
-		ensure(bIsDataInitialize);
-
-		TSharedPtr<FJsonObject> CmdObject = MakeShared<FJsonObject>();
-		TSharedPtr<FJsonObject> ActionDataObject = MakeShared<FJsonObject>();
-		//CmdObject
-		CmdObject->SetStringField(GetCommandIDKey(), GetAction());
-		CmdObject->SetStringField(GetTranslatorIDKey(), GetTranslatorID());
-		ActionDataObject->SetStringField(GetSourceFilenameKey(), GetSourceFilename());
-		CmdObject->SetObjectField(GetCommandDataKey(), ActionDataObject);
-
-		FString LoadSourceCmd;
-		TSharedRef< TJsonWriter< TCHAR, TPrettyJsonPrintPolicy<TCHAR> > > JsonWriter = TJsonWriterFactory< TCHAR, TPrettyJsonPrintPolicy<TCHAR> >::Create(&LoadSourceCmd);
-		if (!FJsonSerializer::Serialize(CmdObject.ToSharedRef(), JsonWriter))
+		FString FJsonLoadSourceCmd::ToJson() const
 		{
-			//Error creating the json cmd string 
-			return FString();
-		}
-		return LoadSourceCmd;
-	}
+			//Code should not do a ToJson if the data was not set before
+			ensure(bIsDataInitialize);
 
-	bool FJsonLoadSourceCmd::FromJson(const FString& JsonString)
-	{
-		TSharedRef<TJsonReader<TCHAR>> Reader = FJsonStringReader::Create(JsonString);
+			TSharedPtr<FJsonObject> CmdObject = MakeShared<FJsonObject>();
+			TSharedPtr<FJsonObject> ActionDataObject = MakeShared<FJsonObject>();
+			//CmdObject
+			CmdObject->SetStringField(GetCommandIDKey(), GetAction());
+			CmdObject->SetStringField(GetTranslatorIDKey(), GetTranslatorID());
+			ActionDataObject->SetStringField(GetSourceFilenameKey(), GetSourceFilename());
+			CmdObject->SetObjectField(GetCommandDataKey(), ActionDataObject);
 
-		TSharedPtr<FJsonObject> CmdObject;
-		if (!FJsonSerializer::Deserialize(Reader, CmdObject) || !CmdObject.IsValid())
-		{
-			//Cannot read the json file
-			return false;
-		}
-		FString JsonActionValue;
-		if (!CmdObject->TryGetStringField(GetCommandIDKey(), JsonActionValue))
-		{
-			//The json cmd id key is missing
-			return false;
+			FString LoadSourceCmd;
+			TSharedRef< TJsonWriter< TCHAR, TPrettyJsonPrintPolicy<TCHAR> > > JsonWriter = TJsonWriterFactory< TCHAR, TPrettyJsonPrintPolicy<TCHAR> >::Create(&LoadSourceCmd);
+			if (!FJsonSerializer::Serialize(CmdObject.ToSharedRef(), JsonWriter))
+			{
+				//Error creating the json cmd string 
+				return FString();
+			}
+			return LoadSourceCmd;
 		}
 
-		if (JsonActionValue != GetAction())
+		bool FJsonLoadSourceCmd::FromJson(const FString& JsonString)
 		{
-			//This json do not represent a load command
-			return false;
-		}
+			TSharedRef<TJsonReader<TCHAR>> Reader = FJsonStringReader::Create(JsonString);
 
-		//Read the json
-		if (!CmdObject->TryGetStringField(GetTranslatorIDKey(), TranslatorID))
-		{
-			//Missing Load command translator ID
-			return false;
-		}
-		const TSharedPtr<FJsonObject>* ActionDataObject = nullptr;
-		if (!CmdObject->TryGetObjectField(GetCommandDataKey(), ActionDataObject))
-		{
-			//Missing Load Action data object
-			return false;
-		}
-		if (!((*ActionDataObject)->TryGetStringField(GetSourceFilenameKey(), SourceFilename)))
-		{
-			return false;
-		}
+			TSharedPtr<FJsonObject> CmdObject;
+			if (!FJsonSerializer::Deserialize(Reader, CmdObject) || !CmdObject.IsValid())
+			{
+				//Cannot read the json file
+				return false;
+			}
+			FString JsonActionValue;
+			if (!CmdObject->TryGetStringField(GetCommandIDKey(), JsonActionValue))
+			{
+				//The json cmd id key is missing
+				return false;
+			}
 
-		//Since we filled the data from the json file, set the data has been initialize.
-		bIsDataInitialize = true;
-		return true;
-	}
+			if (JsonActionValue != GetAction())
+			{
+				//This json do not represent a load command
+				return false;
+			}
 
-	FString FJsonLoadSourceCmd::JsonResultParser::ToJson() const
-	{
-		TSharedPtr<FJsonObject> ResultObject = MakeShared<FJsonObject>();
-		//CmdObject
-		ResultObject->SetStringField(GetResultFilenameKey(), GetResultFilename());
+			//Read the json
+			if (!CmdObject->TryGetStringField(GetTranslatorIDKey(), TranslatorID))
+			{
+				//Missing Load command translator ID
+				return false;
+			}
+			const TSharedPtr<FJsonObject>* ActionDataObject = nullptr;
+			if (!CmdObject->TryGetObjectField(GetCommandDataKey(), ActionDataObject))
+			{
+				//Missing Load Action data object
+				return false;
+			}
+			if (!((*ActionDataObject)->TryGetStringField(GetSourceFilenameKey(), SourceFilename)))
+			{
+				return false;
+			}
 
-		FString JsonResult;
-		TSharedRef< TJsonWriter< TCHAR, TPrettyJsonPrintPolicy<TCHAR> > > JsonWriter = TJsonWriterFactory< TCHAR, TPrettyJsonPrintPolicy<TCHAR> >::Create(&JsonResult);
-		if (!FJsonSerializer::Serialize(ResultObject.ToSharedRef(), JsonWriter))
-		{
-			//Error creating the json cmd string 
-			return FString();
+			//Since we filled the data from the json file, set the data has been initialize.
+			bIsDataInitialize = true;
+			return true;
 		}
-		return JsonResult;
-	}
 
-	bool FJsonLoadSourceCmd::JsonResultParser::FromJson(const FString& JsonString)
-	{
-		TSharedRef<TJsonReader<TCHAR>> Reader = FJsonStringReader::Create(JsonString);
+		FString FJsonLoadSourceCmd::JsonResultParser::ToJson() const
+		{
+			TSharedPtr<FJsonObject> ResultObject = MakeShared<FJsonObject>();
+			//CmdObject
+			ResultObject->SetStringField(GetResultFilenameKey(), GetResultFilename());
 
-		TSharedPtr<FJsonObject> ResultObject;
-		if (!FJsonSerializer::Deserialize(Reader, ResultObject) || !ResultObject.IsValid())
+			FString JsonResult;
+			TSharedRef< TJsonWriter< TCHAR, TPrettyJsonPrintPolicy<TCHAR> > > JsonWriter = TJsonWriterFactory< TCHAR, TPrettyJsonPrintPolicy<TCHAR> >::Create(&JsonResult);
+			if (!FJsonSerializer::Serialize(ResultObject.ToSharedRef(), JsonWriter))
+			{
+				//Error creating the json cmd string 
+				return FString();
+			}
+			return JsonResult;
+		}
+
+		bool FJsonLoadSourceCmd::JsonResultParser::FromJson(const FString& JsonString)
 		{
-			//Cannot read the json file
-			return false;
+			TSharedRef<TJsonReader<TCHAR>> Reader = FJsonStringReader::Create(JsonString);
+
+			TSharedPtr<FJsonObject> ResultObject;
+			if (!FJsonSerializer::Deserialize(Reader, ResultObject) || !ResultObject.IsValid())
+			{
+				//Cannot read the json file
+				return false;
+			}
+			return ResultObject->TryGetStringField(GetResultFilenameKey(), ResultFilename);
 		}
-		return ResultObject->TryGetStringField(GetResultFilenameKey(), ResultFilename);
-	}
 
-} // NS InterchangeDispatcher
+	} //ns Interchange
+}//ns UE
diff --git a/Engine/Plugins/Runtime/Interchange/Source/Dispatcher/Private/InterchangeWorkerHandler.cpp b/Engine/Plugins/Runtime/Interchange/Source/Dispatcher/Private/InterchangeWorkerHandler.cpp
--- a/Engine/Plugins/Runtime/Interchange/Source/Dispatcher/Private/InterchangeWorkerHandler.cpp
+++ b/Engine/Plugins/Runtime/Interchange/Source/Dispatcher/Private/InterchangeWorkerHandler.cpp
@@ -12,335 +12,338 @@
 #include "Sockets.h"
 #include "SocketSubsystem.h"
 
-namespace InterchangeDispatcher
+namespace UE
 {
-
-static FString GetWorkerExecutablePath()
-{
-	static FString ProcessorPath = [&]()
+	namespace Interchange
 	{
-		FString Path = FPaths::Combine(FPaths::EnginePluginsDir(), TEXT("Runtime/Interchange/Binaries"));
+
+		static FString GetWorkerExecutablePath()
+		{
+			static FString ProcessorPath = [&]()
+			{
+				FString Path = FPaths::Combine(FPaths::EnginePluginsDir(), TEXT("Runtime/Interchange/Binaries"));
 
 #if PLATFORM_MAC
-		Path = FPaths::Combine(Path, TEXT("Mac/InterchangeWorker"));
+				Path = FPaths::Combine(Path, TEXT("Mac/InterchangeWorker"));
 #elif PLATFORM_LINUX
-		Path = FPaths::Combine(Path, TEXT("Linux/InterchangeWorker"));
+				Path = FPaths::Combine(Path, TEXT("Linux/InterchangeWorker"));
 #elif PLATFORM_WINDOWS
-	#if UE_BUILD_DEBUG
-		Path = FPaths::Combine(Path, TEXT("Win64/InterchangeWorker-Win64-Debug.exe"));
-	#else
-		Path = FPaths::Combine(Path, TEXT("Win64/InterchangeWorker.exe"));
-	#endif
+#if UE_BUILD_DEBUG
+				Path = FPaths::Combine(Path, TEXT("Win64/InterchangeWorker-Win64-Debug.exe"));
+#else
+				Path = FPaths::Combine(Path, TEXT("Win64/InterchangeWorker.exe"));
 #endif
-		if (!FPaths::FileExists(Path))
-		{
-			UE_LOG(LogInterchangeDispatcher, Warning, TEXT("InterchangeWorker executable not found. Expected location: %s"), *FPaths::ConvertRelativePathToFull(ProcessorPath));
-		}
-		return Path;
-	}();
-
-	return ProcessorPath;
-}
-
-FInterchangeWorkerHandler::FInterchangeWorkerHandler(FInterchangeDispatcher& InDispatcher, FString& InResultFolder)
-	: Dispatcher(InDispatcher)
-	, WorkerState(EWorkerState::Uninitialized)
-	, ErrorState(EWorkerErrorState::Ok)
-	, ResultFolder(InResultFolder)
-	, bShouldTerminate(false)
-{
-	//Use a randomGuid to generate unique name
-	int32 UniqueID = (FPlatformTime::Cycles64() & 0x00000000EFFFFFFF);
-	ThreadName = FString(TEXT("InterchangeWorkerHandler_")) + FString::FromInt(UniqueID);
-	IOThread = FThread(*ThreadName, [this]() { Run(); } );
-}
-
-FInterchangeWorkerHandler::~FInterchangeWorkerHandler()
-{
-	IOThread.Join();
-}
+#endif
+				if (!FPaths::FileExists(Path))
+				{
+					UE_LOG(LogInterchangeDispatcher, Warning, TEXT("InterchangeWorker executable not found. Expected location: %s"), *FPaths::ConvertRelativePathToFull(ProcessorPath));
+				}
+				return Path;
+			}();
 
-void FInterchangeWorkerHandler::StartWorkerProcess()
-{
-	uint32 WorkerProcessId = 0;
+			return ProcessorPath;
+		}
 
-	ensure(ErrorState == EWorkerErrorState::Ok);
-	FString ProcessorPath = GetWorkerExecutablePath();
-	if (FPaths::FileExists(ProcessorPath))
-	{
-		int32 ListenPort = NetworkInterface.GetListeningPort();
-		if (ListenPort == 0)
+		FInterchangeWorkerHandler::FInterchangeWorkerHandler(FInterchangeDispatcher& InDispatcher, FString& InResultFolder)
+			: Dispatcher(InDispatcher)
+			, WorkerState(EWorkerState::Uninitialized)
+			, ErrorState(EWorkerErrorState::Ok)
+			, ResultFolder(InResultFolder)
+			, bShouldTerminate(false)
 		{
-			ErrorState = EWorkerErrorState::ConnectionFailed_NotBound;
-			return;
+			//Use a randomGuid to generate unique name
+			int32 UniqueID = (FPlatformTime::Cycles64() & 0x00000000EFFFFFFF);
+			ThreadName = FString(TEXT("InterchangeWorkerHandler_")) + FString::FromInt(UniqueID);
+			IOThread = FThread(*ThreadName, [this]() { Run(); });
 		}
 
-		FString CommandToProcess;
-
-		// Manually set worker BaseDir of worker process (as automatic deduction is broken in our case,
-		// probably due to the abuse of ExeBinariesSubFolder in InterchangeWorker.Target.cs)
-		// This fixes paths of logs generated by the worker.
-		FString BaseDir = FPlatformProcess::BaseDir();
-		CommandToProcess += TEXT(" -basedir=\"") + BaseDir + TEXT("\"");
-		CommandToProcess += TEXT(" -ServerPID ") + FString::FromInt(FPlatformProcess::GetCurrentProcessId());
-		CommandToProcess += TEXT(" -ServerPort ") + FString::FromInt(ListenPort);
-		CommandToProcess += TEXT(" -InterchangeDispatcherVersion \"") + DispatcherCommandVersion::ToString() + TEXT('"');
-		CommandToProcess += TEXT(" -ResultFolder \"") + ResultFolder + TEXT("\"");
-		UE_LOG(LogInterchangeDispatcher, Verbose, TEXT("CommandToProcess: %s"), *CommandToProcess);
-		WorkerHandle = FPlatformProcess::CreateProc(*ProcessorPath, *CommandToProcess, true, false, false, &WorkerProcessId, 0, nullptr, nullptr);
-	}
-
-	if (!WorkerHandle.IsValid() || !FPlatformProcess::IsProcRunning(WorkerHandle) || !FPlatformProcess::IsApplicationRunning(WorkerProcessId))
-	{
-		ErrorState = EWorkerErrorState::WorkerProcess_CantCreate;
-		return;
-	}
-}
-
-void FInterchangeWorkerHandler::ValidateConnection()
-{
-	if (!NetworkInterface.IsValid())
-	{
-		UE_LOG(LogInterchangeDispatcher, Error, TEXT("NetworkInterface lost"));
-		WorkerState = EWorkerState::Closing;
-		ErrorState = EWorkerErrorState::ConnectionLost;
-	}
-	else if (WorkerHandle.IsValid() && !FPlatformProcess::IsProcRunning(WorkerHandle))
-	{
-		UE_LOG(LogInterchangeDispatcher, Error, TEXT("Worker lost"));
-		WorkerState = EWorkerState::Closing;
-		ErrorState = EWorkerErrorState::WorkerProcess_Lost;
-	}
-}
-
-void FInterchangeWorkerHandler::Run()
-{
-	WorkerState = EWorkerState::Uninitialized;
-	RunInternal();
-	WorkerState = EWorkerState::Terminated;
-	UE_CLOG(ErrorState != EWorkerErrorState::Ok, LogInterchangeDispatcher, Warning, TEXT("Handler ended with error: %s"), EWorkerErrorStateAsString(ErrorState));
-}
-
-void FInterchangeWorkerHandler::RunInternal()
-{
-	while (IsAlive())
-	{
-		switch (WorkerState)
+		FInterchangeWorkerHandler::~FInterchangeWorkerHandler()
 		{
-			case EWorkerState::Uninitialized:
-			{
-				ErrorState = EWorkerErrorState::Ok;
-
-				StartWorkerProcess();
+			IOThread.Join();
+		}
 
-				if (ErrorState != EWorkerErrorState::Ok)
-				{
-					WorkerState = EWorkerState::Terminated;
-					break;
-				}
+		void FInterchangeWorkerHandler::StartWorkerProcess()
+		{
+			uint32 WorkerProcessId = 0;
 
-				// The Accept() call on the server blocks until a connection is initiated from a client
-				static const FString SocketDescription = TEXT("InterchangeWorkerHandler");
-				if (!NetworkInterface.Accept(SocketDescription, Config::AcceptTimeout_s))
-				{
-					ErrorState = EWorkerErrorState::ConnectionFailed_NoClient;
-				}
-				else
+			ensure(ErrorState == EWorkerErrorState::Ok);
+			FString ProcessorPath = GetWorkerExecutablePath();
+			if (FPaths::FileExists(ProcessorPath))
+			{
+				int32 ListenPort = NetworkInterface.GetListeningPort();
+				if (ListenPort == 0)
 				{
-					CommandIO.SetNetworkInterface(&NetworkInterface);
+					ErrorState = EWorkerErrorState::ConnectionFailed_NotBound;
+					return;
 				}
 
-				if (ErrorState != EWorkerErrorState::Ok)
-				{
-					WorkerState = EWorkerState::Closing;
-					break;
-				}
+				FString CommandToProcess;
+
+				// Manually set worker BaseDir of worker process (as automatic deduction is broken in our case,
+				// probably due to the abuse of ExeBinariesSubFolder in InterchangeWorker.Target.cs)
+				// This fixes paths of logs generated by the worker.
+				FString BaseDir = FPlatformProcess::BaseDir();
+				CommandToProcess += TEXT(" -basedir=\"") + BaseDir + TEXT("\"");
+				CommandToProcess += TEXT(" -ServerPID ") + FString::FromInt(FPlatformProcess::GetCurrentProcessId());
+				CommandToProcess += TEXT(" -ServerPort ") + FString::FromInt(ListenPort);
+				CommandToProcess += TEXT(" -InterchangeDispatcherVersion \"") + DispatcherCommandVersion::ToString() + TEXT('"');
+				CommandToProcess += TEXT(" -ResultFolder \"") + ResultFolder + TEXT("\"");
+				UE_LOG(LogInterchangeDispatcher, Verbose, TEXT("CommandToProcess: %s"), *CommandToProcess);
+				WorkerHandle = FPlatformProcess::CreateProc(*ProcessorPath, *CommandToProcess, true, false, false, &WorkerProcessId, 0, nullptr, nullptr);
+			}
 
-				WorkerState = EWorkerState::Idle;
-				break;
+			if (!WorkerHandle.IsValid() || !FPlatformProcess::IsProcRunning(WorkerHandle) || !FPlatformProcess::IsApplicationRunning(WorkerProcessId))
+			{
+				ErrorState = EWorkerErrorState::WorkerProcess_CantCreate;
+				return;
 			}
+		}
 
-			case EWorkerState::Idle:
+		void FInterchangeWorkerHandler::ValidateConnection()
+		{
+			if (!NetworkInterface.IsValid())
 			{
-				// Fetch a new task
-				ensureMsgf(CurrentTask.IsSet() == false, TEXT("We should not have a current task when fetching a new one"));
-				CurrentTask = Dispatcher.GetNextTask();
+				UE_LOG(LogInterchangeDispatcher, Error, TEXT("NetworkInterface lost"));
+				WorkerState = EWorkerState::Closing;
+				ErrorState = EWorkerErrorState::ConnectionLost;
+			}
+			else if (WorkerHandle.IsValid() && !FPlatformProcess::IsProcRunning(WorkerHandle))
+			{
+				UE_LOG(LogInterchangeDispatcher, Error, TEXT("Worker lost"));
+				WorkerState = EWorkerState::Closing;
+				ErrorState = EWorkerErrorState::WorkerProcess_Lost;
+			}
+		}
 
-				if (CurrentTask.IsSet())
-				{
-					FRunTaskCommand NewTask(CurrentTask.GetValue());
+		void FInterchangeWorkerHandler::Run()
+		{
+			WorkerState = EWorkerState::Uninitialized;
+			RunInternal();
+			WorkerState = EWorkerState::Terminated;
+			UE_CLOG(ErrorState != EWorkerErrorState::Ok, LogInterchangeDispatcher, Warning, TEXT("Handler ended with error: %s"), EWorkerErrorStateAsString(ErrorState));
+		}
 
-					if (CommandIO.SendCommand(NewTask, Config::SendCommandTimeout_s))
-					{
-						UE_LOG(LogInterchangeDispatcher, Verbose, TEXT("New task command sent"));
-						WorkerState = EWorkerState::Processing;
-					}
-					else
+		void FInterchangeWorkerHandler::RunInternal()
+		{
+			while (IsAlive())
+			{
+				switch (WorkerState)
+				{
+					case EWorkerState::Uninitialized:
 					{
-						// Signal that the Task was not processed
-						TArray<FString> GarbageMessages;
-						Dispatcher.SetTaskState(CurrentTask->Index, ETaskState::UnTreated, FString(), GarbageMessages);
+						ErrorState = EWorkerErrorState::Ok;
 
-						UE_LOG(LogInterchangeDispatcher, Error, TEXT("New task command issue"));
-						WorkerState = EWorkerState::Closing;
-						ErrorState = EWorkerErrorState::ConnectionLost_SendFailed;
-					}
-				}
-				else if (bShouldTerminate)
-				{
-					UE_LOG(LogInterchangeDispatcher, Verbose, TEXT("Exit loop gracefully"));
-					WorkerState = EWorkerState::Closing;
-				}
-				else
-				{
-					ValidateConnection();
+						StartWorkerProcess();
 
-					// consume
-					if (TSharedPtr<ICommand> Command = CommandIO.GetNextCommand(Config::IdleLoopDelay))
-					{
-						ProcessCommand(*Command);
-					}
-				}
+						if (ErrorState != EWorkerErrorState::Ok)
+						{
+							WorkerState = EWorkerState::Terminated;
+							break;
+						}
 
-				break;
-			}
+						// The Accept() call on the server blocks until a connection is initiated from a client
+						static const FString SocketDescription = TEXT("InterchangeWorkerHandler");
+						if (!NetworkInterface.Accept(SocketDescription, Config::AcceptTimeout_s))
+						{
+							ErrorState = EWorkerErrorState::ConnectionFailed_NoClient;
+						}
+						else
+						{
+							CommandIO.SetNetworkInterface(&NetworkInterface);
+						}
 
-			case EWorkerState::Processing:
-			{
-				if (TSharedPtr<ICommand> Command = CommandIO.GetNextCommand(Config::ProcessingLoopDelay))
-				{
-					ProcessCommand(*Command);
+						if (ErrorState != EWorkerErrorState::Ok)
+						{
+							WorkerState = EWorkerState::Closing;
+							break;
+						}
 
-					bool bProcessingOver = CurrentTask.IsSet() == false;
-					if (bProcessingOver)
-					{
-						WorkerState = bShouldTerminate ? EWorkerState::Closing : EWorkerState::Idle;
+						WorkerState = EWorkerState::Idle;
+						break;
 					}
-				}
-				else
-				{
-					ValidateConnection();
-					if (ErrorState == EWorkerErrorState::WorkerProcess_Lost)
+
+					case EWorkerState::Idle:
 					{
+						// Fetch a new task
+						ensureMsgf(CurrentTask.IsSet() == false, TEXT("We should not have a current task when fetching a new one"));
+						CurrentTask = Dispatcher.GetNextTask();
+
 						if (CurrentTask.IsSet())
 						{
-							TArray<FString> GarbageMessages;
-							Dispatcher.SetTaskState(CurrentTask->Index, ETaskState::ProcessFailed, FString(), GarbageMessages);
-							CurrentTask.Reset();
+							FRunTaskCommand NewTask(CurrentTask.GetValue());
+
+							if (CommandIO.SendCommand(NewTask, Config::SendCommandTimeout_s))
+							{
+								UE_LOG(LogInterchangeDispatcher, Verbose, TEXT("New task command sent"));
+								WorkerState = EWorkerState::Processing;
+							}
+							else
+							{
+								// Signal that the Task was not processed
+								TArray<FString> GarbageMessages;
+								Dispatcher.SetTaskState(CurrentTask->Index, ETaskState::UnTreated, FString(), GarbageMessages);
+
+								UE_LOG(LogInterchangeDispatcher, Error, TEXT("New task command issue"));
+								WorkerState = EWorkerState::Closing;
+								ErrorState = EWorkerErrorState::ConnectionLost_SendFailed;
+							}
 						}
+						else if (bShouldTerminate)
+						{
+							UE_LOG(LogInterchangeDispatcher, Verbose, TEXT("Exit loop gracefully"));
+							WorkerState = EWorkerState::Closing;
+						}
+						else
+						{
+							ValidateConnection();
 
-						WorkerState = EWorkerState::Closing;
-					}
-				}
-				break;
-			}
+							// consume
+							if (TSharedPtr<ICommand> Command = CommandIO.GetNextCommand(Config::IdleLoopDelay))
+							{
+								ProcessCommand(*Command);
+							}
+						}
 
-			case EWorkerState::Closing:
-			{
-				// try to close the process gracefully
-				if (WorkerHandle.IsValid())
-				{
-					bool CloseByCommand = Config::CloseProcessByCommand && CommandIO.IsValid() && FPlatformProcess::IsProcRunning(WorkerHandle);
+						break;
+					}
 
-					bool bClosed = false;
-					if (CloseByCommand)
+					case EWorkerState::Processing:
 					{
-						FTerminateCommand Terminate;
-						CommandIO.SendCommand(Terminate, 0);
+						if (TSharedPtr<ICommand> Command = CommandIO.GetNextCommand(Config::ProcessingLoopDelay))
+						{
+							ProcessCommand(*Command);
 
-						for (int32 i = 0; i < int32(10. * Config::TerminateTimeout_s); ++i)
+							bool bProcessingOver = CurrentTask.IsSet() == false;
+							if (bProcessingOver)
+							{
+								WorkerState = bShouldTerminate ? EWorkerState::Closing : EWorkerState::Idle;
+							}
+						}
+						else
 						{
-							if (!FPlatformProcess::IsProcRunning(WorkerHandle))
+							ValidateConnection();
+							if (ErrorState == EWorkerErrorState::WorkerProcess_Lost)
 							{
-								bClosed = true;
-								break;
+								if (CurrentTask.IsSet())
+								{
+									TArray<FString> GarbageMessages;
+									Dispatcher.SetTaskState(CurrentTask->Index, ETaskState::ProcessFailed, FString(), GarbageMessages);
+									CurrentTask.Reset();
+								}
+
+								WorkerState = EWorkerState::Closing;
 							}
-							FPlatformProcess::Sleep(0.1);
 						}
+						break;
 					}
 
-					if (!bClosed)
+					case EWorkerState::Closing:
 					{
-						FPlatformProcess::TerminateProc(WorkerHandle, true);
+						// try to close the process gracefully
+						if (WorkerHandle.IsValid())
+						{
+							bool CloseByCommand = Config::CloseProcessByCommand && CommandIO.IsValid() && FPlatformProcess::IsProcRunning(WorkerHandle);
+
+							bool bClosed = false;
+							if (CloseByCommand)
+							{
+								FTerminateCommand Terminate;
+								CommandIO.SendCommand(Terminate, 0);
+
+								for (int32 i = 0; i < int32(10. * Config::TerminateTimeout_s); ++i)
+								{
+									if (!FPlatformProcess::IsProcRunning(WorkerHandle))
+									{
+										bClosed = true;
+										break;
+									}
+									FPlatformProcess::Sleep(0.1);
+								}
+							}
+
+							if (!bClosed)
+							{
+								FPlatformProcess::TerminateProc(WorkerHandle, true);
+							}
+						}
+
+						// Process commands still in input queue
+						CommandIO.Disconnect(0);
+						while (TSharedPtr<ICommand> Command = CommandIO.GetNextCommand(0))
+						{
+							ProcessCommand(*Command);
+						}
+
+						WorkerState = EWorkerState::Terminated;
+						break;
 					}
-				}
 
-				// Process commands still in input queue
-				CommandIO.Disconnect(0);
-				while (TSharedPtr<ICommand> Command = CommandIO.GetNextCommand(0))
-				{
-					ProcessCommand(*Command);
+					default:
+					{
+						ensureMsgf(false, TEXT("missing case handling"));
+					}
 				}
-
-				WorkerState = EWorkerState::Terminated;
-				break;
 			}
+		}
 
-			default:
-			{
-				ensureMsgf(false, TEXT("missing case handling"));
-			}
+		void FInterchangeWorkerHandler::Stop()
+		{
+			bShouldTerminate = true;
 		}
-	}
-}
 
-void FInterchangeWorkerHandler::Stop()
-{
-	bShouldTerminate = true;
-}
+		bool FInterchangeWorkerHandler::IsAlive() const
+		{
+			return WorkerState != EWorkerState::Terminated;
+		}
 
-bool FInterchangeWorkerHandler::IsAlive() const
-{
-	return WorkerState != EWorkerState::Terminated;
-}
+		void FInterchangeWorkerHandler::ProcessCommand(ICommand& Command)
+		{
+			switch (Command.GetType())
+			{
+				case ECommandId::Ping:
+					ProcessCommand(StaticCast<FPingCommand&>(Command));
+					break;
 
-void FInterchangeWorkerHandler::ProcessCommand(ICommand& Command)
-{
-	switch (Command.GetType())
-	{
-		case ECommandId::Ping:
-			ProcessCommand(StaticCast<FPingCommand&>(Command));
-			break;
+				case ECommandId::NotifyEndTask:
+					ProcessCommand(StaticCast<FCompletedTaskCommand&>(Command));
+					break;
 
-		case ECommandId::NotifyEndTask:
-			ProcessCommand(StaticCast<FCompletedTaskCommand&>(Command));
-			break;
+				default:
+					break;
+			}
+		}
 
-		default:
-			break;
-	}
-}
+		void FInterchangeWorkerHandler::ProcessCommand(FPingCommand& PingCommand)
+		{
+			UE::Interchange::FBackPingCommand BackPing;
+			CommandIO.SendCommand(BackPing, 0);
+		}
 
-void FInterchangeWorkerHandler::ProcessCommand(FPingCommand& PingCommand)
-{
-	InterchangeDispatcher::FBackPingCommand BackPing;
-	CommandIO.SendCommand(BackPing, 0);
-}
+		void FInterchangeWorkerHandler::ProcessCommand(FCompletedTaskCommand& CompletedTaskCommand)
+		{
+			if (!CurrentTask.IsSet())
+			{
+				return;
+			}
 
-void FInterchangeWorkerHandler::ProcessCommand(FCompletedTaskCommand& CompletedTaskCommand)
-{
-	if (!CurrentTask.IsSet())
-	{
-		return;
-	}
+			Dispatcher.SetTaskState(CurrentTask->Index, CompletedTaskCommand.ProcessResult, CompletedTaskCommand.JSonResult, CompletedTaskCommand.JSonMessages);
+			CurrentTask.Reset();
+		}
 
-	Dispatcher.SetTaskState(CurrentTask->Index, CompletedTaskCommand.ProcessResult, CompletedTaskCommand.JSonResult, CompletedTaskCommand.JSonMessages);
-	CurrentTask.Reset();
-}
+		const TCHAR* FInterchangeWorkerHandler::EWorkerErrorStateAsString(EWorkerErrorState Error)
+		{
+			switch (Error)
+			{
+				case EWorkerErrorState::Ok: return TEXT("Ok");
+				case EWorkerErrorState::ConnectionFailed_NotBound: return TEXT("Connection failed (socket not bound)");
+				case EWorkerErrorState::ConnectionFailed_NoClient: return TEXT("Connection failed (No connection from client)");
+				case EWorkerErrorState::ConnectionLost: return TEXT("Connection lost");
+				case EWorkerErrorState::ConnectionLost_SendFailed: return TEXT("Connection lost (send failed)");
+				case EWorkerErrorState::WorkerProcess_CantCreate: return TEXT("Worker process issue (cannot create worker process)");
+				case EWorkerErrorState::WorkerProcess_Lost: return TEXT("Worker process issue (worker lost)");
+				default: return TEXT("unknown");
+			}
+		}
 
-const TCHAR* FInterchangeWorkerHandler::EWorkerErrorStateAsString(EWorkerErrorState Error)
-{
-	switch (Error)
-	{
-		case EWorkerErrorState::Ok: return TEXT("Ok");
-		case EWorkerErrorState::ConnectionFailed_NotBound: return TEXT("Connection failed (socket not bound)");
-		case EWorkerErrorState::ConnectionFailed_NoClient: return TEXT("Connection failed (No connection from client)");
-		case EWorkerErrorState::ConnectionLost: return TEXT("Connection lost");
-		case EWorkerErrorState::ConnectionLost_SendFailed: return TEXT("Connection lost (send failed)");
-		case EWorkerErrorState::WorkerProcess_CantCreate: return TEXT("Worker process issue (cannot create worker process)");
-		case EWorkerErrorState::WorkerProcess_Lost: return TEXT("Worker process issue (worker lost)");
-		default: return TEXT("unknown");
-	}
-}
-
-} // ns InterchangeDispatcher
+	} //ns Interchange
+}//ns UE
diff --git a/Engine/Plugins/Runtime/Interchange/Source/Dispatcher/Public/InterchangeCommands.h b/Engine/Plugins/Runtime/Interchange/Source/Dispatcher/Public/InterchangeCommands.h
--- a/Engine/Plugins/Runtime/Interchange/Source/Dispatcher/Public/InterchangeCommands.h
+++ b/Engine/Plugins/Runtime/Interchange/Source/Dispatcher/Public/InterchangeCommands.h
@@ -8,140 +8,143 @@
 
 class FArchive;
 
-namespace InterchangeDispatcher
+namespace UE
 {
-
-class DispatcherCommandVersion
-{
-public:
-	//Major version should be updated when there is an existing API that has been change.
-	static int32 GetMajor() {return 1;}
-	//Minor version must be updated when there is an addition to the API.
-	static int32 GetMinor() { return 0; }
-	//Patch version should be update if there is some bug fixes in the private code.
-	static int32 GetPatch() { return 0; }
-
-	/** Return the version in string format "Major.Minor.Patch" */
-	static FString ToString()
+	namespace Interchange
 	{
-		static FString Version = FString::FromInt(GetMajor()) + TEXT(".") + FString::FromInt(GetMinor()) + TEXT(".") + FString::FromInt(GetPatch());
-		return Version;
-	}
 
-	static bool FromString(const FString& VersionStr, int32& OutMajor, int32& OutMinor, int32& OutPatch)
-	{
-		TArray<FString> Tokens;
-		VersionStr.ParseIntoArray(Tokens, TEXT("."));
-		if (Tokens.Num() != 3)
+		class DispatcherCommandVersion
 		{
-			return false;
-		}
-		OutMajor = FCString::Atoi(*(Tokens[0]));
-		OutMinor = FCString::Atoi(*(Tokens[1]));
-		OutPatch = FCString::Atoi(*(Tokens[2]));
-		return true;
-	}
-
-	/**
-	 * We consider having the same major and minor version will make the API fully compatible.
-	 * Patch is only a hint in case you need a particular fixes
-	 * @param Major - The Major version you want to test against the compile Major version
-	 * @param Minor - The Minor version you want to test against the compile Minor version
-	 * @param Patch - In case we need it later, not use currently
-	 *
-	 * @return true if Major and Minor version equal the compile one, InPatch is only there in case we need it in a future release
-	 */
-	static bool IsAPICompatible(int32 Major, int32 Minor, int32 Patch)
-	{
-		return (GetMajor() == Major && GetMinor() == Minor);
-	}
-};
+		public:
+			//Major version should be updated when there is an existing API that has been change.
+			static int32 GetMajor() { return 1; }
+			//Minor version must be updated when there is an addition to the API.
+			static int32 GetMinor() { return 0; }
+			//Patch version should be update if there is some bug fixes in the private code.
+			static int32 GetPatch() { return 0; }
+
+			/** Return the version in string format "Major.Minor.Patch" */
+			static FString ToString()
+			{
+				static FString Version = FString::FromInt(GetMajor()) + TEXT(".") + FString::FromInt(GetMinor()) + TEXT(".") + FString::FromInt(GetPatch());
+				return Version;
+			}
+
+			static bool FromString(const FString& VersionStr, int32& OutMajor, int32& OutMinor, int32& OutPatch)
+			{
+				TArray<FString> Tokens;
+				VersionStr.ParseIntoArray(Tokens, TEXT("."));
+				if (Tokens.Num() != 3)
+				{
+					return false;
+				}
+				OutMajor = FCString::Atoi(*(Tokens[0]));
+				OutMinor = FCString::Atoi(*(Tokens[1]));
+				OutPatch = FCString::Atoi(*(Tokens[2]));
+				return true;
+			}
+
+			/**
+			 * We consider having the same major and minor version will make the API fully compatible.
+			 * Patch is only a hint in case you need a particular fixes
+			 * @param Major - The Major version you want to test against the compile Major version
+			 * @param Minor - The Minor version you want to test against the compile Minor version
+			 * @param Patch - In case we need it later, not use currently
+			 *
+			 * @return true if Major and Minor version equal the compile one, InPatch is only there in case we need it in a future release
+			 */
+			static bool IsAPICompatible(int32 Major, int32 Minor, int32 Patch)
+			{
+				return (GetMajor() == Major && GetMinor() == Minor);
+			}
+		};
+
+		enum class ECommandId : uint8
+		{
+			Invalid,
+			Ping,
+			BackPing,
+			RunTask,
+			NotifyEndTask,
+			Terminate,
+			Last
+		};
+
+		class ICommand
+		{
+		public:
+			virtual ~ICommand() = default;
 
-enum class ECommandId : uint8
-{
-	Invalid,
-	Ping,
-	BackPing,
-	RunTask,
-	NotifyEndTask,
-	Terminate,
-	Last
-};
-
-class ICommand
-{
-public:
-	virtual ~ICommand() = default;
+			virtual ECommandId GetType() const = 0;
 
-	virtual ECommandId GetType() const = 0;
+			friend void operator<<(FArchive& Ar, ICommand& C) { C.SerializeImpl(Ar); }
 
-	friend void operator<<(FArchive& Ar, ICommand& C) { C.SerializeImpl(Ar); }
+		protected:
+			virtual void SerializeImpl(FArchive&) {}
+		};
 
-protected:
-	virtual void SerializeImpl(FArchive&) {}
-};
-	
 
 
-// Create a new command from its type
-TSharedPtr<ICommand> CreateCommand(ECommandId CommandType);
+		// Create a new command from its type
+		TSharedPtr<ICommand> CreateCommand(ECommandId CommandType);
 
-// Converts a command into a byte buffer
-void SerializeCommand(ICommand& Command, TArray<uint8>& OutBuffer);
+		// Converts a command into a byte buffer
+		void SerializeCommand(ICommand& Command, TArray<uint8>& OutBuffer);
 
-// Converts byte buffer back into a Command
-// returns nullptr in case of error
-TSharedPtr<ICommand> DeserializeCommand(const TArray<uint8>& InBuffer);
+		// Converts byte buffer back into a Command
+		// returns nullptr in case of error
+		TSharedPtr<ICommand> DeserializeCommand(const TArray<uint8>& InBuffer);
 
 
 
-class FTerminateCommand : public ICommand
-{
-public:
-	virtual ECommandId GetType() const override { return ECommandId::Terminate; }
-};
+		class FTerminateCommand : public ICommand
+		{
+		public:
+			virtual ECommandId GetType() const override { return ECommandId::Terminate; }
+		};
 
 
-class FPingCommand : public ICommand
-{
-public:
-	virtual ECommandId GetType() const override { return ECommandId::Ping; }
-};
+		class FPingCommand : public ICommand
+		{
+		public:
+			virtual ECommandId GetType() const override { return ECommandId::Ping; }
+		};
 
 
-class FBackPingCommand : public ICommand
-{
-public:
-	virtual ECommandId GetType() const override { return ECommandId::BackPing; }
-};
+		class FBackPingCommand : public ICommand
+		{
+		public:
+			virtual ECommandId GetType() const override { return ECommandId::BackPing; }
+		};
 
-class FRunTaskCommand : public ICommand
-{
-public:
-	FRunTaskCommand() = default;
-	FRunTaskCommand(const FTask& Task) : JsonDescription(Task.JsonDescription), JobIndex(Task.Index) {}
-	virtual ECommandId GetType() const override { return ECommandId::RunTask; }
+		class FRunTaskCommand : public ICommand
+		{
+		public:
+			FRunTaskCommand() = default;
+			FRunTaskCommand(const FTask& Task) : JsonDescription(Task.JsonDescription), JobIndex(Task.Index) {}
+			virtual ECommandId GetType() const override { return ECommandId::RunTask; }
 
-protected:
-	virtual void SerializeImpl(FArchive&) override;
+		protected:
+			virtual void SerializeImpl(FArchive&) override;
 
-public:
-	FString JsonDescription;
-	int32 JobIndex = -1;
-};
+		public:
+			FString JsonDescription;
+			int32 JobIndex = -1;
+		};
 
-class FCompletedTaskCommand : public ICommand
-{
-public:
-	virtual ECommandId GetType() const override { return ECommandId::NotifyEndTask; }
+		class FCompletedTaskCommand : public ICommand
+		{
+		public:
+			virtual ECommandId GetType() const override { return ECommandId::NotifyEndTask; }
 
-protected:
-	virtual void SerializeImpl(FArchive&) override;
+		protected:
+			virtual void SerializeImpl(FArchive&) override;
 
-public:
-	ETaskState ProcessResult = ETaskState::Unknown;
-	FString JSonResult;
-	TArray<FString> JSonMessages;
-};
+		public:
+			ETaskState ProcessResult = ETaskState::Unknown;
+			FString JSonResult;
+			TArray<FString> JSonMessages;
+		};
 
-} // ns InterchangeDispatcher
+	} //ns Interchange
+}//ns UE
diff --git a/Engine/Plugins/Runtime/Interchange/Source/Dispatcher/Public/InterchangeDispatcherConfig.h b/Engine/Plugins/Runtime/Interchange/Source/Dispatcher/Public/InterchangeDispatcherConfig.h
--- a/Engine/Plugins/Runtime/Interchange/Source/Dispatcher/Public/InterchangeDispatcherConfig.h
+++ b/Engine/Plugins/Runtime/Interchange/Source/Dispatcher/Public/InterchangeDispatcherConfig.h
@@ -1,34 +1,37 @@
 // Copyright Epic Games, Inc. All Rights Reserved.
 
 #pragma once
-
-namespace InterchangeDispatcher
+namespace UE
 {
-	namespace Config
+	namespace Interchange
 	{
-		// Server side timeout for connection with the client
-		constexpr double AcceptTimeout_s = 3.0;
+		namespace Config
+		{
+			// Server side timeout for connection with the client
+			constexpr double AcceptTimeout_s = 3.0;
+
+			// Client side timeout for connection with the server
+			constexpr double ConnectTimeout_s = 3.0;
 
-		// Client side timeout for connection with the server
-		constexpr double ConnectTimeout_s = 3.0;
+			// Maximum time we accept to wait when a command is sent.
+			// If the command is not sent by this delay, the operation is considered as failed
+			constexpr double SendCommandTimeout_s = 1.0;
 
-		// Maximum time we accept to wait when a command is sent.
-		// If the command is not sent by this delay, the operation is considered as failed
-		constexpr double SendCommandTimeout_s = 1.0;
+			// When true, a Terminate command is sent to the processor and we wait for
+			// TerminateTimeout_s seconds for it to close, before a hard TerminateProc.
+			// When false, the process is directly closed with a TerminateProc.
+			constexpr bool CloseProcessByCommand = true;
 
-		// When true, a Terminate command is sent to the processor and we wait for
-		// TerminateTimeout_s seconds for it to close, before a hard TerminateProc.
-		// When false, the process is directly closed with a TerminateProc.
-		constexpr bool CloseProcessByCommand = true;
+			// Maximum time we accept to wait for the process to close after a Terminate
+			// command has been sent.
+			constexpr double TerminateTimeout_s = 1.0;
 
-		// Maximum time we accept to wait for the process to close after a Terminate
-		// command has been sent.
-		constexpr double TerminateTimeout_s = 1.0;
+			// Delay between successive checks for a new task, and validation of the connection
+			constexpr double IdleLoopDelay = 0.1;
 
-		// Delay between successive checks for a new task, and validation of the connection
-		constexpr double IdleLoopDelay = 0.1;
+			// Delay between successive validation of the connection
+			constexpr double ProcessingLoopDelay = 0.5;
+		};
+	} //ns Interchange
+}//ns UE
 
-		// Delay between successive validation of the connection
-		constexpr double ProcessingLoopDelay = 0.5;
-	};
-} // ns InterchangeDispatcher
diff --git a/Engine/Plugins/Runtime/Interchange/Source/Dispatcher/Public/InterchangeDispatcherNetworking.h b/Engine/Plugins/Runtime/Interchange/Source/Dispatcher/Public/InterchangeDispatcherNetworking.h
--- a/Engine/Plugins/Runtime/Interchange/Source/Dispatcher/Public/InterchangeDispatcherNetworking.h
+++ b/Engine/Plugins/Runtime/Interchange/Source/Dispatcher/Public/InterchangeDispatcherNetworking.h
@@ -8,118 +8,120 @@
 #include "Containers/Queue.h"
 #include "Sockets.h"
 
-
-namespace InterchangeDispatcher
+namespace UE
 {
+	namespace Interchange
+	{
 
-enum class SocketErrorCode
-{
-	NoError = 0,
-	Error_Create,
-	Error_Bind,
-	Error_Listen,
-	UnableToReadOnSocket,
-	UnableToSendData,
-	CouldNotStartWSA,
-	UnableToGetLocalAddress,
-	ConnectionToServerFailed,
-};
-
-
-struct FMessageHeader
-{
-	int32 ByteSize = -1;
+		enum class SocketErrorCode
+		{
+			NoError = 0,
+			Error_Create,
+			Error_Bind,
+			Error_Listen,
+			UnableToReadOnSocket,
+			UnableToSendData,
+			CouldNotStartWSA,
+			UnableToGetLocalAddress,
+			ConnectionToServerFailed,
+		};
 
-	friend void operator<< (FArchive& Ar, FMessageHeader& H)
-	{
-		FString Guard = Ar.IsLoading() ? TEXT("") : TEXT("guard");
-		Ar << Guard;
-		Ar << H.ByteSize;
-		ensure(Guard == TEXT("guard"));
-	}
-};
 
+		struct FMessageHeader
+		{
+			int32 ByteSize = -1;
 
-class FNetworkNode
-{
-public:
-	virtual ~FNetworkNode();
+			friend void operator<< (FArchive& Ar, FMessageHeader& H)
+			{
+				FString Guard = Ar.IsLoading() ? TEXT("") : TEXT("guard");
+				Ar << Guard;
+				Ar << H.ByteSize;
+				ensure(Guard == TEXT("guard"));
+			}
+		};
 
-	// Send a byte buffer as an atomic message
-	// returns true when a message is successfully sent
-	bool SendMessage(const TArray<uint8>& Buffer, double Timeout_s);
 
-	// Receive a byte buffer as an atomic message
-	// returns true when a message is fetched. In that case, OutBuffer contains the message.
-	bool ReceiveMessage(TArray<uint8>& OutBuffer, double Timeout_s);
+		class FNetworkNode
+		{
+		public:
+			virtual ~FNetworkNode();
 
-	bool IsValid() const { return !bReadError && !bWriteError; }
+			// Send a byte buffer as an atomic message
+			// returns true when a message is successfully sent
+			bool SendMessage(const TArray<uint8>& Buffer, double Timeout_s);
 
-protected:
-	struct FMessage
-	{
-		FMessageHeader Header;
-		TArray<uint8> Content;
-	};
-	FSocket* CreateInternalSocket(const FString& Description);
-	void CloseSocket(FSocket*& Socket);
-	bool IsConnected();
+			// Receive a byte buffer as an atomic message
+			// returns true when a message is fetched. In that case, OutBuffer contains the message.
+			bool ReceiveMessage(TArray<uint8>& OutBuffer, double Timeout_s);
 
-protected:
-	FSocket* ConnectedSocket = nullptr;
-	SocketErrorCode ConnectedSocketError = SocketErrorCode::NoError;
+			bool IsValid() const { return !bReadError && !bWriteError; }
 
-	FMessage IncommingMessage;
+		protected:
+			struct FMessage
+			{
+				FMessageHeader Header;
+				TArray<uint8> Content;
+			};
+			FSocket* CreateInternalSocket(const FString& Description);
+			void CloseSocket(FSocket*& Socket);
+			bool IsConnected();
 
-	bool bReadError = false;
-	bool bWriteError = false;
-};
+		protected:
+			FSocket* ConnectedSocket = nullptr;
+			SocketErrorCode ConnectedSocketError = SocketErrorCode::NoError;
 
+			FMessage IncommingMessage;
 
+			bool bReadError = false;
+			bool bWriteError = false;
+		};
 
-// Create, Bind, Listen, Accept sequence
-class FNetworkServerNode : public FNetworkNode
-{
-public:
-	FNetworkServerNode();
-	~FNetworkServerNode();
-	int32 GetListeningPort();
-	bool Accept(const FString& Description, double Timeout_s);
 
-private:
-	FSocket* ListeningSocket;
-};
 
+		// Create, Bind, Listen, Accept sequence
+		class FNetworkServerNode : public FNetworkNode
+		{
+		public:
+			FNetworkServerNode();
+			~FNetworkServerNode();
+			int32 GetListeningPort();
+			bool Accept(const FString& Description, double Timeout_s);
 
+		private:
+			FSocket* ListeningSocket;
+		};
 
-// Create, Connect sequence
-class FNetworkClientNode : public FNetworkNode
-{
-public:
-	bool Connect(const FString& Description, int32 ServerPort, double Timeout_s);
-};
 
 
+		// Create, Connect sequence
+		class FNetworkClientNode : public FNetworkNode
+		{
+		public:
+			bool Connect(const FString& Description, int32 ServerPort, double Timeout_s);
+		};
 
-class FCommandQueue
-{
-public:
-	void SetNetworkInterface(FNetworkNode* InNetworkInterface);
 
-	TSharedPtr<ICommand> GetNextCommand(double Timeout_s);
-	bool SendCommand(ICommand& Commmand, double Timeout_s);
 
-	bool IsValid() { return NetworkInterface ? NetworkInterface->IsValid() : true;}
+		class FCommandQueue
+		{
+		public:
+			void SetNetworkInterface(FNetworkNode* InNetworkInterface);
+
+			TSharedPtr<ICommand> GetNextCommand(double Timeout_s);
+			bool SendCommand(ICommand& Commmand, double Timeout_s);
+
+			bool IsValid() { return NetworkInterface ? NetworkInterface->IsValid() : true; }
 
-	void Disconnect(double Timeout_s);
+			void Disconnect(double Timeout_s);
 
-private:
-	bool Poll(double Timeout_s);
-	TQueue<TSharedPtr<ICommand>> InCommands;
+		private:
+			bool Poll(double Timeout_s);
+			TQueue<TSharedPtr<ICommand>> InCommands;
 
-private:
-	FNetworkNode* NetworkInterface = nullptr;
-};
+		private:
+			FNetworkNode* NetworkInterface = nullptr;
+		};
 
 
-} // ns InterchangeDispatcher
+	} //ns Interchange
+}//ns UE
diff --git a/Engine/Plugins/Runtime/Interchange/Source/Dispatcher/Public/InterchangeDispatcherTask.h b/Engine/Plugins/Runtime/Interchange/Source/Dispatcher/Public/InterchangeDispatcherTask.h
--- a/Engine/Plugins/Runtime/Interchange/Source/Dispatcher/Public/InterchangeDispatcherTask.h
+++ b/Engine/Plugins/Runtime/Interchange/Source/Dispatcher/Public/InterchangeDispatcherTask.h
@@ -3,146 +3,148 @@
 #pragma once
 
 #include "CoreMinimal.h"
-
-namespace InterchangeDispatcher
-{
-enum class ETaskState
-{
-	Unknown,
-	Running,
-	UnTreated,
-	ProcessOk,
-	ProcessFailed,
-};
-
-struct FTask
-{
-	FTask() = default;
-
-	FTask(const FString& InJsonDescription)
-	{
-		JsonDescription = InJsonDescription;
-		State = ETaskState::UnTreated;
-	}
-
-	FString JsonDescription;
-	int32 Index = -1;
-	ETaskState State = ETaskState::Unknown;
-	FString JsonResult;
-	TArray<FString> JsonMessages;
-};
-
-/**
- * Json cmd helper to be able to read and write a FTask::JsonDescription
- */
-class INTERCHANGEDISPATCHER_API IJsonCmdBase
-{
-public:
-	virtual ~IJsonCmdBase() = default;
-
-	virtual FString GetAction() const = 0;
-	virtual FString GetTranslatorID() const = 0;
-	virtual FString ToJson() const = 0;
-
-	/**
-	 * Return false if the JsonString do not match the command, true otherwise.
-	 */
-	virtual bool FromJson(const FString& JsonString) = 0;
-
-	static FString GetCommandIDKey()
-	{
-		static const FString Key = TEXT("CmdID");
-		return Key;
-	}
-	static FString GetTranslatorIDKey()
-	{
-		static const FString Key = TEXT("TranslatorID");
-		return Key;
-	}
-	static FString GetCommandDataKey()
-	{
-		static const FString Key = TEXT("CmdData");
-		return Key;
-	}
-
-protected:
-	//Use this member to know if the data is initialize before using it
-	bool bIsDataInitialize = false;
-};
-
-class INTERCHANGEDISPATCHER_API FJsonLoadSourceCmd : public IJsonCmdBase
+namespace UE
 {
-public:
-	FJsonLoadSourceCmd()
+	namespace Interchange
 	{
-		bIsDataInitialize = false;
-	}
-
-	FJsonLoadSourceCmd(const FString& InTranslatorID, const FString& InSourceFilename)
-		: TranslatorID(InTranslatorID)
-		, SourceFilename(InSourceFilename)
-	{
-		bIsDataInitialize = true;
-	}
-
-	virtual FString GetAction() const override
-	{
-		static const FString LoadString = TEXT("Load");
-		return LoadString;
-	}
-
-	virtual FString GetTranslatorID() const override
-	{
-		//Code should not do query data if the data was not set before
-		ensure(bIsDataInitialize);
-		return TranslatorID;
-	}
-
-	virtual FString ToJson() const;
-	virtual bool FromJson(const FString& JsonString);
-
-	FString GetSourceFilename() const
-	{
-		//Code should not do query data if the data was not set before
-		ensure(bIsDataInitialize);
-		return SourceFilename;
-	}
-
-	static FString GetSourceFilenameKey()
-	{
-		static const FString Key = TEXT("SourceFile");
-		return Key;
-	}
-
-	/**
-	 * Use this class helper to create the cmd result json string and to read it
-	 */
-	class INTERCHANGEDISPATCHER_API JsonResultParser
-	{
-	public:
-		FString GetResultFilename() const
+		enum class ETaskState
 		{
-			return ResultFilename;
-		}
-		void SetResultFilename(const FString& InResultFilename)
+			Unknown,
+			Running,
+			UnTreated,
+			ProcessOk,
+			ProcessFailed,
+		};
+
+		struct FTask
 		{
-			ResultFilename = InResultFilename;
-		}
-		FString ToJson() const;
-		bool FromJson(const FString& JsonString);
-
-		static FString GetResultFilenameKey()
+			FTask() = default;
+
+			FTask(const FString& InJsonDescription)
+			{
+				JsonDescription = InJsonDescription;
+				State = ETaskState::UnTreated;
+			}
+
+			FString JsonDescription;
+			int32 Index = -1;
+			ETaskState State = ETaskState::Unknown;
+			FString JsonResult;
+			TArray<FString> JsonMessages;
+		};
+
+		/**
+		 * Json cmd helper to be able to read and write a FTask::JsonDescription
+		 */
+		class INTERCHANGEDISPATCHER_API IJsonCmdBase
+		{
+		public:
+			virtual ~IJsonCmdBase() = default;
+
+			virtual FString GetAction() const = 0;
+			virtual FString GetTranslatorID() const = 0;
+			virtual FString ToJson() const = 0;
+
+			/**
+			 * Return false if the JsonString do not match the command, true otherwise.
+			 */
+			virtual bool FromJson(const FString& JsonString) = 0;
+
+			static FString GetCommandIDKey()
+			{
+				static const FString Key = TEXT("CmdID");
+				return Key;
+			}
+			static FString GetTranslatorIDKey()
+			{
+				static const FString Key = TEXT("TranslatorID");
+				return Key;
+			}
+			static FString GetCommandDataKey()
+			{
+				static const FString Key = TEXT("CmdData");
+				return Key;
+			}
+
+		protected:
+			//Use this member to know if the data is initialize before using it
+			bool bIsDataInitialize = false;
+		};
+
+		class INTERCHANGEDISPATCHER_API FJsonLoadSourceCmd : public IJsonCmdBase
 		{
-			const FString Key = TEXT("ResultFile");
-			return Key;
-		}
-	private:
-		FString ResultFilename = FString();
-	};
-
-private:
-	FString TranslatorID = FString();
-	FString SourceFilename = FString();
-};
-
-} // NS InterchangeDispatcher
+		public:
+			FJsonLoadSourceCmd()
+			{
+				bIsDataInitialize = false;
+			}
+
+			FJsonLoadSourceCmd(const FString& InTranslatorID, const FString& InSourceFilename)
+				: TranslatorID(InTranslatorID)
+				, SourceFilename(InSourceFilename)
+			{
+				bIsDataInitialize = true;
+			}
+
+			virtual FString GetAction() const override
+			{
+				static const FString LoadString = TEXT("Load");
+				return LoadString;
+			}
+
+			virtual FString GetTranslatorID() const override
+			{
+				//Code should not do query data if the data was not set before
+				ensure(bIsDataInitialize);
+				return TranslatorID;
+			}
+
+			virtual FString ToJson() const;
+			virtual bool FromJson(const FString& JsonString);
+
+			FString GetSourceFilename() const
+			{
+				//Code should not do query data if the data was not set before
+				ensure(bIsDataInitialize);
+				return SourceFilename;
+			}
+
+			static FString GetSourceFilenameKey()
+			{
+				static const FString Key = TEXT("SourceFile");
+				return Key;
+			}
+
+			/**
+			 * Use this class helper to create the cmd result json string and to read it
+			 */
+			class INTERCHANGEDISPATCHER_API JsonResultParser
+			{
+			public:
+				FString GetResultFilename() const
+				{
+					return ResultFilename;
+				}
+				void SetResultFilename(const FString& InResultFilename)
+				{
+					ResultFilename = InResultFilename;
+				}
+				FString ToJson() const;
+				bool FromJson(const FString& JsonString);
+
+				static FString GetResultFilenameKey()
+				{
+					const FString Key = TEXT("ResultFile");
+					return Key;
+				}
+			private:
+				FString ResultFilename = FString();
+			};
+
+		private:
+			FString TranslatorID = FString();
+			FString SourceFilename = FString();
+		};
+
+	} //ns Interchange
+}//ns UE
diff --git a/Engine/Plugins/Runtime/Interchange/Source/Dispatcher/Public/InterchangeWorkerHandler.h b/Engine/Plugins/Runtime/Interchange/Source/Dispatcher/Public/InterchangeWorkerHandler.h
--- a/Engine/Plugins/Runtime/Interchange/Source/Dispatcher/Public/InterchangeWorkerHandler.h
+++ b/Engine/Plugins/Runtime/Interchange/Source/Dispatcher/Public/InterchangeWorkerHandler.h
@@ -11,71 +11,72 @@
 #include "HAL/PlatformProcess.h"
 #include "HAL/Thread.h"
 
-
-namespace InterchangeDispatcher
+namespace UE
 {
-
-class FInterchangeDispatcher;
-
-//Handle a Worker by socket communication
-class INTERCHANGEDISPATCHER_API FInterchangeWorkerHandler
-{
-	enum class EWorkerState
-	{
-		Uninitialized,
-		Idle, // Initialized, available for processing
-		Processing, // Currently processing a task
-		Closing, // in the process of terminating
-		Terminated, // aka. Not Alive
-	};
-
-	enum class EWorkerErrorState
+	namespace Interchange
 	{
-		Ok,
-		ConnectionFailed_NotBound,
-		ConnectionFailed_NoClient,
-		ConnectionLost,
-		ConnectionLost_SendFailed,
-		WorkerProcess_CantCreate,
-		WorkerProcess_Lost,
-	};
-
-public:
-	FInterchangeWorkerHandler(FInterchangeDispatcher& InDispatcher, FString& InResultFolder);
-	~FInterchangeWorkerHandler();
-
-	void Run();
-	bool IsAlive() const;
-	void Stop();
-
-private:
-	void RunInternal();
-	void StartWorkerProcess();
-	void ValidateConnection();
-
-	void ProcessCommand(ICommand& Command);
-	void ProcessCommand(FPingCommand& PingCommand);
-	void ProcessCommand(FCompletedTaskCommand& RunTaskCommand);
-	const TCHAR* EWorkerErrorStateAsString(EWorkerErrorState e);
-
-private:
-	FInterchangeDispatcher& Dispatcher;
-
-	// Send and receive commands
-	FNetworkServerNode NetworkInterface;
-	FCommandQueue CommandIO;
-	FThread IOThread;
-	FString ThreadName;
-
-	// External process
-	FProcHandle WorkerHandle;
-	TAtomic<EWorkerState> WorkerState;
-	EWorkerErrorState ErrorState;
-
-	// self
-	FString ResultFolder;
-	TOptional<FTask> CurrentTask;
-	bool bShouldTerminate;
-
-};
-} // namespace InterchangeDispatcher
+		class FInterchangeDispatcher;
+
+		//Handle a Worker by socket communication
+		class INTERCHANGEDISPATCHER_API FInterchangeWorkerHandler
+		{
+			enum class EWorkerState
+			{
+				Uninitialized,
+				Idle, // Initialized, available for processing
+				Processing, // Currently processing a task
+				Closing, // in the process of terminating
+				Terminated, // aka. Not Alive
+			};
+
+			enum class EWorkerErrorState
+			{
+				Ok,
+				ConnectionFailed_NotBound,
+				ConnectionFailed_NoClient,
+				ConnectionLost,
+				ConnectionLost_SendFailed,
+				WorkerProcess_CantCreate,
+				WorkerProcess_Lost,
+			};
+
+		public:
+			FInterchangeWorkerHandler(FInterchangeDispatcher& InDispatcher, FString& InResultFolder);
+			~FInterchangeWorkerHandler();
+
+			void Run();
+			bool IsAlive() const;
+			void Stop();
+
+		private:
+			void RunInternal();
+			void StartWorkerProcess();
+			void ValidateConnection();
+
+			void ProcessCommand(ICommand& Command);
+			void ProcessCommand(FPingCommand& PingCommand);
+			void ProcessCommand(FCompletedTaskCommand& RunTaskCommand);
+			const TCHAR* EWorkerErrorStateAsString(EWorkerErrorState e);
+
+		private:
+			FInterchangeDispatcher& Dispatcher;
+
+			// Send and receive commands
+			FNetworkServerNode NetworkInterface;
+			FCommandQueue CommandIO;
+			FThread IOThread;
+			FString ThreadName;
+
+			// External process
+			FProcHandle WorkerHandle;
+			TAtomic<EWorkerState> WorkerState;
+			EWorkerErrorState ErrorState;
+
+			// self
+			FString ResultFolder;
+			TOptional<FTask> CurrentTask;
+			bool bShouldTerminate;
+
+		};
+	} //ns Interchange
+}//ns UE
