Make sure all interchange namespace respect the coding standard

diff --git a/Engine/Plugins/Runtime/Interchange/Source/Import/Private/Fbx/InterchangeFbxTranslator.cpp b/Engine/Plugins/Runtime/Interchange/Source/Import/Private/Fbx/InterchangeFbxTranslator.cpp
--- a/Engine/Plugins/Runtime/Interchange/Source/Import/Private/Fbx/InterchangeFbxTranslator.cpp
+++ b/Engine/Plugins/Runtime/Interchange/Source/Import/Private/Fbx/InterchangeFbxTranslator.cpp
@@ -67,7 +67,7 @@ bool UInterchangeFbxTranslator::Translate(const UInterchangeSourceData* SourceDa
 		}
 
 		//Create the dispatcher
-		Dispatcher = MakeUnique<InterchangeDispatcher::FInterchangeDispatcher>(ResultFolder);
+		Dispatcher = MakeUnique<UE::Interchange::FInterchangeDispatcher>(ResultFolder);
 
 		if(ensure(Dispatcher.IsValid()))
 		{
@@ -87,19 +87,19 @@ bool UInterchangeFbxTranslator::Translate(const UInterchangeSourceData* SourceDa
 	//Blocking call until all tasks are executed
 	Dispatcher->WaitAllTaskToCompleteExecution();
 		
-	InterchangeDispatcher::ETaskState TaskState;
+	UE::Interchange::ETaskState TaskState;
 	FString JsonResult;
 	TArray<FString> JSonMessages;
 	Dispatcher->GetTaskState(TaskIndex, TaskState, JsonResult, JSonMessages);
 
 	//TODO: Parse the JSonMessage and add the message to the interchange not yet develop error messaging
 
-	if(TaskState != InterchangeDispatcher::ETaskState::ProcessOk)
+	if(TaskState != UE::Interchange::ETaskState::ProcessOk)
 	{
 		return false;
 	}
 	//Grab the result file and fill the BaseNodeContainer
-	InterchangeDispatcher::FJsonLoadSourceCmd::JsonResultParser ResultParser;
+	UE::Interchange::FJsonLoadSourceCmd::JsonResultParser ResultParser;
 	ResultParser.FromJson(JsonResult);
 	FString BaseNodeContainerFilename = ResultParser.GetResultFilename();
 
@@ -118,45 +118,45 @@ void UInterchangeFbxTranslator::ImportFinish()
 }
 
 
-TOptional<Interchange::FImportImage> UInterchangeFbxTranslator::GetTexturePayloadData(const UInterchangeSourceData* SourceData, const FString& PayLoadKey) const
+TOptional<UE::Interchange::FImportImage> UInterchangeFbxTranslator::GetTexturePayloadData(const UInterchangeSourceData* SourceData, const FString& PayLoadKey) const
 {
 	UInterchangeSourceData* PayloadSourceData = UInterchangeManager::GetInterchangeManager().CreateSourceData(PayLoadKey);
 	FGCObjectScopeGuard ScopedSourceData(PayloadSourceData);
 	
 	if (!PayloadSourceData)
 	{
-		return TOptional<Interchange::FImportImage>();
+		return TOptional<UE::Interchange::FImportImage>();
 	}
 	UInterchangeTranslatorBase* SourceTranslator = UInterchangeManager::GetInterchangeManager().GetTranslatorForSourceData(PayloadSourceData);
 	FGCObjectScopeGuard ScopedSourceTranslator(SourceTranslator);
 	const IInterchangeTexturePayloadInterface* TextureTranslator = Cast<IInterchangeTexturePayloadInterface>(SourceTranslator);
 	if (!ensure(TextureTranslator))
 	{
-		return TOptional<Interchange::FImportImage>();
+		return TOptional<UE::Interchange::FImportImage>();
 	}
 	return TextureTranslator->GetTexturePayloadData(PayloadSourceData, PayLoadKey);
 }
 
-TOptional<Interchange::FMaterialPayloadData> UInterchangeFbxTranslator::GetMaterialPayloadData(const UInterchangeSourceData* SourceData, const FString& PayLoadKey) const
+TOptional<UE::Interchange::FMaterialPayloadData> UInterchangeFbxTranslator::GetMaterialPayloadData(const UInterchangeSourceData* SourceData, const FString& PayLoadKey) const
 {
 	//Not implemented, currently we do not have any payload data for the material
-	return TOptional<Interchange::FMaterialPayloadData>();
+	return TOptional<UE::Interchange::FMaterialPayloadData>();
 }
 
-TOptional<Interchange::FStaticMeshPayloadData> UInterchangeFbxTranslator::GetStaticMeshPayloadData(const UInterchangeSourceData* SourceData, const FString& PayLoadKey) const
+TOptional<UE::Interchange::FStaticMeshPayloadData> UInterchangeFbxTranslator::GetStaticMeshPayloadData(const UInterchangeSourceData* SourceData, const FString& PayLoadKey) const
 {
 	//Not implemented, currently we do not have any payload data for static meshes
-	return TOptional<Interchange::FStaticMeshPayloadData>();
+	return TOptional<UE::Interchange::FStaticMeshPayloadData>();
 }
 
-TOptional<Interchange::FSkeletalMeshPayloadData> UInterchangeFbxTranslator::GetSkeletalMeshPayloadData(const UInterchangeSourceData* SourceData, const FString& PayLoadKey) const
+TOptional<UE::Interchange::FSkeletalMeshPayloadData> UInterchangeFbxTranslator::GetSkeletalMeshPayloadData(const UInterchangeSourceData* SourceData, const FString& PayLoadKey) const
 {
 	//Not implemented, currently we do not have any payload data for skeletal meshes
-	return TOptional<Interchange::FSkeletalMeshPayloadData>();
+	return TOptional<UE::Interchange::FSkeletalMeshPayloadData>();
 }
 
 FString UInterchangeFbxTranslator::CreateLoadFbxFileCommand(const FString& FbxFilePath) const
 {
-	InterchangeDispatcher::FJsonLoadSourceCmd LoadSourceCommand(TEXT("FBX"), FbxFilePath);
+	UE::Interchange::FJsonLoadSourceCmd LoadSourceCommand(TEXT("FBX"), FbxFilePath);
 	return LoadSourceCommand.ToJson();
 }
\ No newline at end of file
diff --git a/Engine/Plugins/Runtime/Interchange/Source/Import/Private/Texture/InterchangeBMPTranslator.cpp b/Engine/Plugins/Runtime/Interchange/Source/Import/Private/Texture/InterchangeBMPTranslator.cpp
--- a/Engine/Plugins/Runtime/Interchange/Source/Import/Private/Texture/InterchangeBMPTranslator.cpp
+++ b/Engine/Plugins/Runtime/Interchange/Source/Import/Private/Texture/InterchangeBMPTranslator.cpp
@@ -17,7 +17,6 @@
 #include "Texture/TextureTranslatorUtilities.h"
 
 
-
 bool UInterchangeBMPTranslator::CanImportSourceData(const UInterchangeSourceData* InSourceData) const
 {
 	FString Extension = FPaths::GetExtension(InSourceData->GetFilename());
@@ -27,15 +26,15 @@ bool UInterchangeBMPTranslator::CanImportSourceData(const UInterchangeSourceData
 
 bool UInterchangeBMPTranslator::Translate(const UInterchangeSourceData* SourceData, UInterchangeBaseNodeContainer& BaseNodeContainer) const
 {
-	return Interchange::FTextureTranslatorUtilities::Generic2DTextureTranslate(SourceData, BaseNodeContainer);
+	return UE::Interchange::FTextureTranslatorUtilities::Generic2DTextureTranslate(SourceData, BaseNodeContainer);
 }
 
-TOptional<Interchange::FImportImage> UInterchangeBMPTranslator::GetTexturePayloadData(const UInterchangeSourceData* SourceData, const FString& PayLoadKey) const
+TOptional<UE::Interchange::FImportImage> UInterchangeBMPTranslator::GetTexturePayloadData(const UInterchangeSourceData* SourceData, const FString& PayLoadKey) const
 {
 	if (!SourceData)
 	{
 		UE_LOG(LogInterchangeImportPlugin, Error, TEXT("Failed to import BMP, bad source data."));
-		return TOptional<Interchange::FImportImage>();
+		return TOptional<UE::Interchange::FImportImage>();
 	}
 
 	TArray64<uint8> SourceDataBuffer;
@@ -45,19 +44,19 @@ TOptional<Interchange::FImportImage> UInterchangeBMPTranslator::GetTexturePayloa
 	if (!Filename.Equals(PayLoadKey))
 	{
 		UE_LOG(LogInterchangeImportPlugin, Error, TEXT("Failed to import BMP, wrong payload key. [%s]"), *Filename);
-		return TOptional<Interchange::FImportImage>();
+		return TOptional<UE::Interchange::FImportImage>();
 	}
 
 	if (!FPaths::FileExists(Filename))
 	{
 		UE_LOG(LogInterchangeImportPlugin, Error, TEXT("Failed to import BMP, cannot open file. [%s]"), *Filename);
-		return TOptional<Interchange::FImportImage>();
+		return TOptional<UE::Interchange::FImportImage>();
 	}
 
 	if (!FFileHelper::LoadFileToArray(SourceDataBuffer, *Filename))
 	{
 		UE_LOG(LogInterchangeImportPlugin, Error, TEXT("Failed to import BMP, cannot load file content into an array. [%s]"), *Filename);
-		return TOptional<Interchange::FImportImage>();
+		return TOptional<UE::Interchange::FImportImage>();
 	}
 
 	const uint8* Buffer = SourceDataBuffer.GetData();
@@ -78,18 +77,18 @@ TOptional<Interchange::FImportImage> UInterchangeBMPTranslator::GetTexturePayloa
 	if (!BmpImageWrapper.IsValid() || !BmpImageWrapper->SetCompressed(Buffer, Length))
 	{
 		UE_LOG(LogInterchangeImportPlugin, Error, TEXT("Failed to decode BMP. [%s]"), *Filename);
-		return TOptional<Interchange::FImportImage>();
+		return TOptional<UE::Interchange::FImportImage>();
 	}
-	if (!Interchange::FImportImageHelper::IsImportResolutionValid(BmpImageWrapper->GetWidth(), BmpImageWrapper->GetHeight(), bAllowNonPowerOfTwo))
+	if (!UE::Interchange::FImportImageHelper::IsImportResolutionValid(BmpImageWrapper->GetWidth(), BmpImageWrapper->GetHeight(), bAllowNonPowerOfTwo))
 	{
 		UE_LOG(LogInterchangeImportPlugin, Error, TEXT("Failed to import BMP, invalid resolution. Resolution[%d, %d], AllowPowerOfTwo[%s], [%s]"), BmpImageWrapper->GetWidth(), BmpImageWrapper->GetHeight(), bAllowNonPowerOfTwo ? TEXT("True") : TEXT("false"), *Filename);
-		return TOptional<Interchange::FImportImage>();
+		return TOptional<UE::Interchange::FImportImage>();
 	}
 
 	int32 BitDepth = BmpImageWrapper->GetBitDepth();
 	ERGBFormat Format = BmpImageWrapper->GetFormat();
 
-	Interchange::FImportImage PayloadData;
+	UE::Interchange::FImportImage PayloadData;
 
 	PayloadData.Init2DWithParams(
 		BmpImageWrapper->GetWidth(),
@@ -101,7 +100,7 @@ TOptional<Interchange::FImportImage> UInterchangeBMPTranslator::GetTexturePayloa
 	if (!BmpImageWrapper->GetRaw(Format, BitDepth, PayloadData.RawData))
 	{
 		UE_LOG(LogInterchangeImportPlugin, Error, TEXT("Failed to decode BMP. [%s]"), *Filename);
-		return TOptional<Interchange::FImportImage>();
+		return TOptional<UE::Interchange::FImportImage>();
 	}
 
 	return PayloadData;
diff --git a/Engine/Plugins/Runtime/Interchange/Source/Import/Private/Texture/InterchangeDDSTranslator.cpp b/Engine/Plugins/Runtime/Interchange/Source/Import/Private/Texture/InterchangeDDSTranslator.cpp
--- a/Engine/Plugins/Runtime/Interchange/Source/Import/Private/Texture/InterchangeDDSTranslator.cpp
+++ b/Engine/Plugins/Runtime/Interchange/Source/Import/Private/Texture/InterchangeDDSTranslator.cpp
@@ -20,10 +20,6 @@
 //////////////////////////////////////////////////////////////////////////
 // DDS helper local function
 
-namespace DDSParserHelper
-{
-} //DDSParserHelper namespace end
-
 bool UInterchangeDDSTranslator::CanImportSourceData(const UInterchangeSourceData* InSourceData) const
 {
 	FString Extension = FPaths::GetExtension(InSourceData->GetFilename());
@@ -61,15 +57,15 @@ bool UInterchangeDDSTranslator::CanImportSourceData(const UInterchangeSourceData
 
 bool UInterchangeDDSTranslator::Translate(const UInterchangeSourceData* SourceData, UInterchangeBaseNodeContainer& BaseNodeContainer) const
 {
-	return Interchange::FTextureTranslatorUtilities::Generic2DTextureTranslate(SourceData, BaseNodeContainer);
+	return UE::Interchange::FTextureTranslatorUtilities::Generic2DTextureTranslate(SourceData, BaseNodeContainer);
 }
 
-TOptional<Interchange::FImportImage> UInterchangeDDSTranslator::GetTexturePayloadData(const UInterchangeSourceData* SourceData, const FString& PayLoadKey) const
+TOptional<UE::Interchange::FImportImage> UInterchangeDDSTranslator::GetTexturePayloadData(const UInterchangeSourceData* SourceData, const FString& PayLoadKey) const
 {
 	if (!SourceData)
 	{
 		UE_LOG(LogInterchangeImportPlugin, Error, TEXT("Failed to import DDS, bad source data."));
-		return TOptional<Interchange::FImportImage>();
+		return TOptional<UE::Interchange::FImportImage>();
 	}
 
 	TArray64<uint8> SourceDataBuffer;
@@ -79,19 +75,19 @@ TOptional<Interchange::FImportImage> UInterchangeDDSTranslator::GetTexturePayloa
 	if (!Filename.Equals(PayLoadKey))
 	{
 		UE_LOG(LogInterchangeImportPlugin, Error, TEXT("Failed to import DDS, wrong payload key. [%s]"), *Filename);
-		return TOptional<Interchange::FImportImage>();
+		return TOptional<UE::Interchange::FImportImage>();
 	}
 
 	if (!FPaths::FileExists(Filename))
 	{
 		UE_LOG(LogInterchangeImportPlugin, Error, TEXT("Failed to import DDS, cannot open file. [%s]"), *Filename);
-		return TOptional<Interchange::FImportImage>();
+		return TOptional<UE::Interchange::FImportImage>();
 	}
 
 	if (!FFileHelper::LoadFileToArray(SourceDataBuffer, *Filename))
 	{
 		UE_LOG(LogInterchangeImportPlugin, Error, TEXT("Failed to import DDS, cannot load file content into an array. [%s]"), *Filename);
-		return TOptional<Interchange::FImportImage>();
+		return TOptional<UE::Interchange::FImportImage>();
 	}
 
 	const uint8* Buffer = SourceDataBuffer.GetData();
@@ -111,13 +107,13 @@ TOptional<Interchange::FImportImage> UInterchangeDDSTranslator::GetTexturePayloa
 	if (!DDSLoadHelper.IsValid2DTexture())
 	{
 		UE_LOG(LogInterchangeImportPlugin, Error, TEXT("Failed to import DDS, unsupported format. [%s]"), *Filename);
-		return TOptional<Interchange::FImportImage>();
+		return TOptional<UE::Interchange::FImportImage>();
 	}
 	// DDS 2d texture
-	if (!Interchange::FImportImageHelper::IsImportResolutionValid(DDSLoadHelper.DDSHeader->dwWidth, DDSLoadHelper.DDSHeader->dwHeight, bAllowNonPowerOfTwo))
+	if (!UE::Interchange::FImportImageHelper::IsImportResolutionValid(DDSLoadHelper.DDSHeader->dwWidth, DDSLoadHelper.DDSHeader->dwHeight, bAllowNonPowerOfTwo))
 	{
 		UE_LOG(LogInterchangeImportPlugin, Error, TEXT("Failed to import PCX, invalid resolution. Resolution[%d, %d], AllowPowerOfTwo[%s], [%s]"), DDSLoadHelper.DDSHeader->dwWidth, DDSLoadHelper.DDSHeader->dwHeight, bAllowNonPowerOfTwo ? TEXT("True") : TEXT("false"), *Filename);
-		return TOptional<Interchange::FImportImage>();
+		return TOptional<UE::Interchange::FImportImage>();
 	}
 
 	ETextureSourceFormat SourceFormat = DDSLoadHelper.ComputeSourceFormat();
@@ -126,17 +122,17 @@ TOptional<Interchange::FImportImage> UInterchangeDDSTranslator::GetTexturePayloa
 	if (SourceFormat == TSF_Invalid)
 	{
 		UE_LOG(LogInterchangeImportPlugin, Error, TEXT("DDS file [%s] contains data in an unsupported format"), *Filename);
-		return TOptional<Interchange::FImportImage>();
+		return TOptional<UE::Interchange::FImportImage>();
 	}
 
 	uint32 MipMapCount = DDSLoadHelper.ComputeMipMapCount();
 	if (MipMapCount <= 0)
 	{
 		UE_LOG(LogInterchangeImportPlugin, Error, TEXT("DDS file [%s] do not have any mipmap"), *Filename);
-		return TOptional<Interchange::FImportImage>();
+		return TOptional<UE::Interchange::FImportImage>();
 	}
 
-	Interchange::FImportImage PayloadData;
+	UE::Interchange::FImportImage PayloadData;
 	PayloadData.Init2DWithMips(
 		DDSLoadHelper.DDSHeader->dwWidth,
 		DDSLoadHelper.DDSHeader->dwHeight,
diff --git a/Engine/Plugins/Runtime/Interchange/Source/Import/Private/Texture/InterchangeEXRTranslator.cpp b/Engine/Plugins/Runtime/Interchange/Source/Import/Private/Texture/InterchangeEXRTranslator.cpp
--- a/Engine/Plugins/Runtime/Interchange/Source/Import/Private/Texture/InterchangeEXRTranslator.cpp
+++ b/Engine/Plugins/Runtime/Interchange/Source/Import/Private/Texture/InterchangeEXRTranslator.cpp
@@ -24,15 +24,15 @@ bool UInterchangeEXRTranslator::CanImportSourceData(const UInterchangeSourceData
 
 bool UInterchangeEXRTranslator::Translate(const UInterchangeSourceData* SourceData, UInterchangeBaseNodeContainer& BaseNodeContainer) const
 {
-	return Interchange::FTextureTranslatorUtilities::Generic2DTextureTranslate(SourceData, BaseNodeContainer);
+	return UE::Interchange::FTextureTranslatorUtilities::Generic2DTextureTranslate(SourceData, BaseNodeContainer);
 }
 
-TOptional<Interchange::FImportImage> UInterchangeEXRTranslator::GetTexturePayloadData(const UInterchangeSourceData* SourceData, const FString& PayLoadKey) const
+TOptional<UE::Interchange::FImportImage> UInterchangeEXRTranslator::GetTexturePayloadData(const UInterchangeSourceData* SourceData, const FString& PayLoadKey) const
 {
 	if (!SourceData)
 	{
 		UE_LOG(LogInterchangeImportPlugin, Error, TEXT("Failed to import EXR, bad source data."));
-		return TOptional<Interchange::FImportImage>();
+		return TOptional<UE::Interchange::FImportImage>();
 	}
 
 	TArray64<uint8> SourceDataBuffer;
@@ -42,19 +42,19 @@ TOptional<Interchange::FImportImage> UInterchangeEXRTranslator::GetTexturePayloa
 	if (!Filename.Equals(PayLoadKey))
 	{
 		UE_LOG(LogInterchangeImportPlugin, Error, TEXT("Failed to import EXR, wrong payload key. [%s]"), *Filename);
-		return TOptional<Interchange::FImportImage>();
+		return TOptional<UE::Interchange::FImportImage>();
 	}
 
 	if (!FPaths::FileExists(Filename))
 	{
 		UE_LOG(LogInterchangeImportPlugin, Error, TEXT("Failed to import EXR, cannot open file. [%s]"), *Filename);
-		return TOptional<Interchange::FImportImage>();
+		return TOptional<UE::Interchange::FImportImage>();
 	}
 
 	if (!FFileHelper::LoadFileToArray(SourceDataBuffer, *Filename))
 	{
 		UE_LOG(LogInterchangeImportPlugin, Error, TEXT("Failed to import EXR, cannot load file content into an array. [%s]"), *Filename);
-		return TOptional<Interchange::FImportImage>();
+		return TOptional<UE::Interchange::FImportImage>();
 	}
 
 	const uint8* Buffer = SourceDataBuffer.GetData();
@@ -75,16 +75,16 @@ TOptional<Interchange::FImportImage> UInterchangeEXRTranslator::GetTexturePayloa
 	if (!ExrImageWrapper.IsValid() || !ExrImageWrapper->SetCompressed(Buffer, Length))
 	{
 		UE_LOG(LogInterchangeImportPlugin, Error, TEXT("Failed to import EXR, unsupported format. [%s]"), *Filename);
-		return TOptional<Interchange::FImportImage>();
+		return TOptional<UE::Interchange::FImportImage>();
 	}
 
 	int32 Width = ExrImageWrapper->GetWidth();
 	int32 Height = ExrImageWrapper->GetHeight();
 
-	if (!Interchange::FImportImageHelper::IsImportResolutionValid(Width, Height, bAllowNonPowerOfTwo))
+	if (!UE::Interchange::FImportImageHelper::IsImportResolutionValid(Width, Height, bAllowNonPowerOfTwo))
 	{
 		UE_LOG(LogInterchangeImportPlugin, Error, TEXT("Failed to import PCX, invalid resolution. Resolution[%d, %d], AllowPowerOfTwo[%s], [%s]"), Width, Height, bAllowNonPowerOfTwo ? TEXT("True") : TEXT("false"), *Filename);
-		return TOptional<Interchange::FImportImage>();
+		return TOptional<UE::Interchange::FImportImage>();
 	}
 
 	// Select the texture's source format
@@ -101,10 +101,10 @@ TOptional<Interchange::FImportImage> UInterchangeEXRTranslator::GetTexturePayloa
 	if (TextureFormat == TSF_Invalid)
 	{
 		UE_LOG(LogInterchangeImportPlugin, Error, TEXT("EXR file [%s] contains data in an unsupported format"), *Filename);
-		return TOptional<Interchange::FImportImage>();
+		return TOptional<UE::Interchange::FImportImage>();
 	}
 
-	Interchange::FImportImage PayloadData;
+	UE::Interchange::FImportImage PayloadData;
 	PayloadData.Init2DWithParams(
 		Width,
 		Height,
@@ -116,7 +116,7 @@ TOptional<Interchange::FImportImage> UInterchangeEXRTranslator::GetTexturePayloa
 	if (!ExrImageWrapper->GetRaw(Format, BitDepth, PayloadData.RawData))
 	{
 		UE_LOG(LogInterchangeImportPlugin, Error, TEXT("Failed to decode EXR. [%s]"), *Filename);
-		return TOptional<Interchange::FImportImage>();
+		return TOptional<UE::Interchange::FImportImage>();
 	}
 
 	return PayloadData;
diff --git a/Engine/Plugins/Runtime/Interchange/Source/Import/Private/Texture/InterchangeJPGTranslator.cpp b/Engine/Plugins/Runtime/Interchange/Source/Import/Private/Texture/InterchangeJPGTranslator.cpp
--- a/Engine/Plugins/Runtime/Interchange/Source/Import/Private/Texture/InterchangeJPGTranslator.cpp
+++ b/Engine/Plugins/Runtime/Interchange/Source/Import/Private/Texture/InterchangeJPGTranslator.cpp
@@ -24,15 +24,15 @@ bool UInterchangeJPGTranslator::CanImportSourceData(const UInterchangeSourceData
 
 bool UInterchangeJPGTranslator::Translate(const UInterchangeSourceData* SourceData, UInterchangeBaseNodeContainer& BaseNodeContainer) const
 {
-	return Interchange::FTextureTranslatorUtilities::Generic2DTextureTranslate(SourceData, BaseNodeContainer);
+	return UE::Interchange::FTextureTranslatorUtilities::Generic2DTextureTranslate(SourceData, BaseNodeContainer);
 }
 
-TOptional<Interchange::FImportImage> UInterchangeJPGTranslator::GetTexturePayloadData(const UInterchangeSourceData* SourceData, const FString& PayLoadKey) const
+TOptional<UE::Interchange::FImportImage> UInterchangeJPGTranslator::GetTexturePayloadData(const UInterchangeSourceData* SourceData, const FString& PayLoadKey) const
 {
 	if (!SourceData)
 	{
 		UE_LOG(LogInterchangeImportPlugin, Error, TEXT("Failed to import JPEG, bad source data."));
-		return TOptional<Interchange::FImportImage>();
+		return TOptional<UE::Interchange::FImportImage>();
 	}
 
 	TArray64<uint8> SourceDataBuffer;
@@ -42,19 +42,19 @@ TOptional<Interchange::FImportImage> UInterchangeJPGTranslator::GetTexturePayloa
 	if (!Filename.Equals(PayLoadKey))
 	{
 		UE_LOG(LogInterchangeImportPlugin, Error, TEXT("Failed to import JPEG, wrong payload key. [%s]"), *Filename);
-		return TOptional<Interchange::FImportImage>();
+		return TOptional<UE::Interchange::FImportImage>();
 	}
 
 	if (!FPaths::FileExists(Filename))
 	{
 		UE_LOG(LogInterchangeImportPlugin, Error, TEXT("Failed to import JPEG, cannot open file. [%s]"), *Filename);
-		return TOptional<Interchange::FImportImage>();
+		return TOptional<UE::Interchange::FImportImage>();
 	}
 
 	if (!FFileHelper::LoadFileToArray(SourceDataBuffer, *Filename))
 	{
 		UE_LOG(LogInterchangeImportPlugin, Error, TEXT("Failed to import JPEG, cannot load file content into an array. [%s]"), *Filename);
-		return TOptional<Interchange::FImportImage>();
+		return TOptional<UE::Interchange::FImportImage>();
 	}
 
 	const uint8* Buffer = SourceDataBuffer.GetData();
@@ -75,12 +75,12 @@ TOptional<Interchange::FImportImage> UInterchangeJPGTranslator::GetTexturePayloa
 	if (!JpegImageWrapper.IsValid() || !JpegImageWrapper->SetCompressed(Buffer, Length))
 	{
 		UE_LOG(LogInterchangeImportPlugin, Error, TEXT("Failed to decode JPEG. [%s]"), *Filename);
-		return TOptional<Interchange::FImportImage>();
+		return TOptional<UE::Interchange::FImportImage>();
 	}
-	if (!Interchange::FImportImageHelper::IsImportResolutionValid(JpegImageWrapper->GetWidth(), JpegImageWrapper->GetHeight(), bAllowNonPowerOfTwo))
+	if (!UE::Interchange::FImportImageHelper::IsImportResolutionValid(JpegImageWrapper->GetWidth(), JpegImageWrapper->GetHeight(), bAllowNonPowerOfTwo))
 	{
 		UE_LOG(LogInterchangeImportPlugin, Error, TEXT("Failed to import JPEG, invalid resolution. Resolution[%d, %d], AllowPowerOfTwo[%s], [%s]"), JpegImageWrapper->GetWidth(), JpegImageWrapper->GetHeight(), bAllowNonPowerOfTwo ? TEXT("True") : TEXT("false"), *Filename);
-		return TOptional<Interchange::FImportImage>();
+		return TOptional<UE::Interchange::FImportImage>();
 	}
 
 	// Select the texture's source format
@@ -110,10 +110,10 @@ TOptional<Interchange::FImportImage> UInterchangeJPGTranslator::GetTexturePayloa
 	if (TextureFormat == TSF_Invalid)
 	{
 		UE_LOG(LogInterchangeImportPlugin, Error, TEXT("JPEG file [%s] contains data in an unsupported format"), *Filename);
-		return TOptional<Interchange::FImportImage>();
+		return TOptional<UE::Interchange::FImportImage>();
 	}
 
-	Interchange::FImportImage PayloadData;
+	UE::Interchange::FImportImage PayloadData;
 
 	PayloadData.Init2DWithParams(
 		JpegImageWrapper->GetWidth(),
@@ -125,7 +125,7 @@ TOptional<Interchange::FImportImage> UInterchangeJPGTranslator::GetTexturePayloa
 	if (!JpegImageWrapper->GetRaw(Format, BitDepth, PayloadData.RawData))
 	{
 		UE_LOG(LogInterchangeImportPlugin, Error, TEXT("Failed to decode JPEG. [%s]"), *Filename);
-		return TOptional<Interchange::FImportImage>();
+		return TOptional<UE::Interchange::FImportImage>();
 	}
 	return PayloadData;
 }
\ No newline at end of file
diff --git a/Engine/Plugins/Runtime/Interchange/Source/Import/Private/Texture/InterchangePCXTranslator.cpp b/Engine/Plugins/Runtime/Interchange/Source/Import/Private/Texture/InterchangePCXTranslator.cpp
--- a/Engine/Plugins/Runtime/Interchange/Source/Import/Private/Texture/InterchangePCXTranslator.cpp
+++ b/Engine/Plugins/Runtime/Interchange/Source/Import/Private/Texture/InterchangePCXTranslator.cpp
@@ -17,49 +17,54 @@
 
 //////////////////////////////////////////////////////////////////////////
 // PCX helper local function
-
-namespace PCXParserHelper
+namespace UE
 {
-	// .PCX file header.
-#pragma pack(push,1)
-	
-	class FPCXFileHeader
+	namespace Interchange
 	{
-	public:
-		uint8	Manufacturer;		// Always 10.
-		uint8	Version;			// PCX file version.
-		uint8	Encoding;			// 1=run-length, 0=none.
-		uint8	BitsPerPixel;		// 1,2,4, or 8.
-		uint16	XMin;				// Dimensions of the image.
-		uint16	YMin;				// Dimensions of the image.
-		uint16	XMax;				// Dimensions of the image.
-		uint16	YMax;				// Dimensions of the image.
-		uint16	XDotsPerInch;		// Horizontal printer resolution.
-		uint16	YDotsPerInch;		// Vertical printer resolution.
-		uint8	OldColorMap[48];	// Old colormap info data.
-		uint8	Reserved1;			// Must be 0.
-		uint8	NumPlanes;			// Number of color planes (1, 3, 4, etc).
-		uint16	BytesPerLine;		// Number of bytes per scanline.
-		uint16	PaletteType;		// How to interpret palette: 1=color, 2=gray.
-		uint16	HScreenSize;		// Horizontal monitor size.
-		uint16	VScreenSize;		// Vertical monitor size.
-		uint8	Reserved2[54];		// Must be 0.
-		friend FArchive& operator<<(FArchive& Ar, FPCXFileHeader& H)
+		namespace Private
 		{
-			Ar << H.Manufacturer << H.Version << H.Encoding << H.BitsPerPixel;
-			Ar << H.XMin << H.YMin << H.XMax << H.YMax << H.XDotsPerInch << H.YDotsPerInch;
-			for (int32 i = 0; i < UE_ARRAY_COUNT(H.OldColorMap); i++)
-				Ar << H.OldColorMap[i];
-			Ar << H.Reserved1 << H.NumPlanes;
-			Ar << H.BytesPerLine << H.PaletteType << H.HScreenSize << H.VScreenSize;
-			for (int32 i = 0; i < UE_ARRAY_COUNT(H.Reserved2); i++)
-				Ar << H.Reserved2[i];
-			return Ar;
-		}
-	};
+			// .PCX file header.
+#pragma pack(push,1)
+
+			class FPCXFileHeader
+			{
+			public:
+				uint8	Manufacturer;		// Always 10.
+				uint8	Version;			// PCX file version.
+				uint8	Encoding;			// 1=run-length, 0=none.
+				uint8	BitsPerPixel;		// 1,2,4, or 8.
+				uint16	XMin;				// Dimensions of the image.
+				uint16	YMin;				// Dimensions of the image.
+				uint16	XMax;				// Dimensions of the image.
+				uint16	YMax;				// Dimensions of the image.
+				uint16	XDotsPerInch;		// Horizontal printer resolution.
+				uint16	YDotsPerInch;		// Vertical printer resolution.
+				uint8	OldColorMap[48];	// Old colormap info data.
+				uint8	Reserved1;			// Must be 0.
+				uint8	NumPlanes;			// Number of color planes (1, 3, 4, etc).
+				uint16	BytesPerLine;		// Number of bytes per scanline.
+				uint16	PaletteType;		// How to interpret palette: 1=color, 2=gray.
+				uint16	HScreenSize;		// Horizontal monitor size.
+				uint16	VScreenSize;		// Vertical monitor size.
+				uint8	Reserved2[54];		// Must be 0.
+				friend FArchive& operator<<(FArchive& Ar, FPCXFileHeader& H)
+				{
+					Ar << H.Manufacturer << H.Version << H.Encoding << H.BitsPerPixel;
+					Ar << H.XMin << H.YMin << H.XMax << H.YMax << H.XDotsPerInch << H.YDotsPerInch;
+					for (int32 i = 0; i < UE_ARRAY_COUNT(H.OldColorMap); i++)
+						Ar << H.OldColorMap[i];
+					Ar << H.Reserved1 << H.NumPlanes;
+					Ar << H.BytesPerLine << H.PaletteType << H.HScreenSize << H.VScreenSize;
+					for (int32 i = 0; i < UE_ARRAY_COUNT(H.Reserved2); i++)
+						Ar << H.Reserved2[i];
+					return Ar;
+				}
+			};
 
-	#pragma pack(pop)
-} //PCXParserHelper namespace end
+#pragma pack(pop)
+		}//ns Private
+	}//ns Interchange
+}//ns UE
 
 bool UInterchangePCXTranslator::CanImportSourceData(const UInterchangeSourceData* InSourceData) const
 {
@@ -70,15 +75,15 @@ bool UInterchangePCXTranslator::CanImportSourceData(const UInterchangeSourceData
 
 bool UInterchangePCXTranslator::Translate(const UInterchangeSourceData* SourceData, UInterchangeBaseNodeContainer& BaseNodeContainer) const
 {
-	return Interchange::FTextureTranslatorUtilities::Generic2DTextureTranslate(SourceData, BaseNodeContainer);
+	return UE::Interchange::FTextureTranslatorUtilities::Generic2DTextureTranslate(SourceData, BaseNodeContainer);
 }
 
-TOptional<Interchange::FImportImage> UInterchangePCXTranslator::GetTexturePayloadData(const UInterchangeSourceData* SourceData, const FString& PayLoadKey) const
+TOptional<UE::Interchange::FImportImage> UInterchangePCXTranslator::GetTexturePayloadData(const UInterchangeSourceData* SourceData, const FString& PayLoadKey) const
 {
 	if (!SourceData)
 	{
 		UE_LOG(LogInterchangeImportPlugin, Error, TEXT("Failed to import PCX, bad source data."));
-		return TOptional<Interchange::FImportImage>();
+		return TOptional<UE::Interchange::FImportImage>();
 	}
 
 	TArray64<uint8> SourceDataBuffer;
@@ -88,19 +93,19 @@ TOptional<Interchange::FImportImage> UInterchangePCXTranslator::GetTexturePayloa
 	if (!Filename.Equals(PayLoadKey))
 	{
 		UE_LOG(LogInterchangeImportPlugin, Error, TEXT("Failed to import PCX, wrong payload key. [%s]"), *Filename);
-		return TOptional<Interchange::FImportImage>();
+		return TOptional<UE::Interchange::FImportImage>();
 	}
 
 	if (!FPaths::FileExists(Filename))
 	{
 		UE_LOG(LogInterchangeImportPlugin, Error, TEXT("Failed to import PCX, cannot open file. [%s]"), *Filename);
-		return TOptional<Interchange::FImportImage>();
+		return TOptional<UE::Interchange::FImportImage>();
 	}
 
 	if (!FFileHelper::LoadFileToArray(SourceDataBuffer, *Filename))
 	{
 		UE_LOG(LogInterchangeImportPlugin, Error, TEXT("Failed to import PCX, cannot load file content into an array. [%s]"), *Filename);
-		return TOptional<Interchange::FImportImage>();
+		return TOptional<UE::Interchange::FImportImage>();
 	}
 
 	const uint8* Buffer = SourceDataBuffer.GetData();
@@ -114,22 +119,22 @@ TOptional<Interchange::FImportImage> UInterchangePCXTranslator::GetTexturePayloa
 
 	IImageWrapperModule& ImageWrapperModule = FModuleManager::LoadModuleChecked<IImageWrapperModule>(FName("ImageWrapper"));
 
-	Interchange::FImportImage PayloadData;
+	UE::Interchange::FImportImage PayloadData;
 
 	//
 	// PCX
 	//
-	const PCXParserHelper::FPCXFileHeader* PCX = (PCXParserHelper::FPCXFileHeader*)Buffer;
-	if (Length >= sizeof(PCXParserHelper::FPCXFileHeader) && PCX->Manufacturer == 10)
+	const UE::Interchange::Private::FPCXFileHeader* PCX = (UE::Interchange::Private::FPCXFileHeader*)Buffer;
+	if (Length >= sizeof(UE::Interchange::Private::FPCXFileHeader) && PCX->Manufacturer == 10)
 	{
 		int32 NewU = PCX->XMax + 1 - PCX->XMin;
 		int32 NewV = PCX->YMax + 1 - PCX->YMin;
 
 		// Check the resolution of the imported texture to ensure validity
-		if (!Interchange::FImportImageHelper::IsImportResolutionValid(NewU, NewV, bAllowNonPowerOfTwo))
+		if (!UE::Interchange::FImportImageHelper::IsImportResolutionValid(NewU, NewV, bAllowNonPowerOfTwo))
 		{
 			UE_LOG(LogInterchangeImportPlugin, Error, TEXT("Failed to import PCX, invalid resolution. Resolution[%d, %d], AllowPowerOfTwo[%s], [%s]"), NewU, NewV, bAllowNonPowerOfTwo ? TEXT("True") : TEXT("false"), *Filename);
-			return TOptional<Interchange::FImportImage>();
+			return TOptional<UE::Interchange::FImportImage>();
 		}
 		else if (PCX->NumPlanes == 1 && PCX->BitsPerPixel == 8)
 		{
@@ -227,13 +232,13 @@ TOptional<Interchange::FImportImage> UInterchangePCXTranslator::GetTexturePayloa
 		else
 		{
 			UE_LOG(LogInterchangeImportPlugin, Error, TEXT("Failed to import PCX, unsupported format. (%i/%i) [%s]"), PCX->NumPlanes, PCX->BitsPerPixel, *Filename);
-			return TOptional<Interchange::FImportImage>();
+			return TOptional<UE::Interchange::FImportImage>();
 		}
 	}
 	else
 	{
 		UE_LOG(LogInterchangeImportPlugin, Error, TEXT("Failed to import PCX, unsupported file version. [%s]"), *Filename);
-		return TOptional<Interchange::FImportImage>();
+		return TOptional<UE::Interchange::FImportImage>();
 	}
 
 	return PayloadData;
diff --git a/Engine/Plugins/Runtime/Interchange/Source/Import/Private/Texture/InterchangePNGTranslator.cpp b/Engine/Plugins/Runtime/Interchange/Source/Import/Private/Texture/InterchangePNGTranslator.cpp
--- a/Engine/Plugins/Runtime/Interchange/Source/Import/Private/Texture/InterchangePNGTranslator.cpp
+++ b/Engine/Plugins/Runtime/Interchange/Source/Import/Private/Texture/InterchangePNGTranslator.cpp
@@ -18,169 +18,174 @@
 
 //////////////////////////////////////////////////////////////////////////
 // PNG helper local function
-
-namespace PNGParserHelper
+namespace UE
 {
-	/**
-	 * This fills any pixels of a texture with have an alpha value of zero,
-	 * with an RGB from the nearest neighboring pixel which has non-zero alpha.
-	 */
-	template<typename PixelDataType, typename ColorDataType, int32 RIdx, int32 GIdx, int32 BIdx, int32 AIdx>
-	class TPNGDataFill
+	namespace Interchange
 	{
-	public:
-
-		TPNGDataFill(int32 SizeX, int32 SizeY, uint8* SourceTextureData)
-			: SourceData(reinterpret_cast<PixelDataType*>(SourceTextureData))
-			, TextureWidth(SizeX)
-			, TextureHeight(SizeY)
-		{
-		}
-
-		void ProcessData()
+		namespace Private
 		{
-			int32 NumZeroedTopRowsToProcess = 0;
-			int32 FillColorRow = -1;
-			for (int32 Y = 0; Y < TextureHeight; ++Y)
+			/**
+			 * This fills any pixels of a texture with have an alpha value of zero,
+			 * with an RGB from the nearest neighboring pixel which has non-zero alpha.
+			 */
+			template<typename PixelDataType, typename ColorDataType, int32 RIdx, int32 GIdx, int32 BIdx, int32 AIdx>
+			class TPNGDataFill
 			{
-				if (!ProcessHorizontalRow(Y))
+			public:
+
+				TPNGDataFill(int32 SizeX, int32 SizeY, uint8* SourceTextureData)
+					: SourceData(reinterpret_cast<PixelDataType*>(SourceTextureData))
+					, TextureWidth(SizeX)
+					, TextureHeight(SizeY)
+				{
+				}
+
+				void ProcessData()
 				{
-					if (FillColorRow != -1)
+					int32 NumZeroedTopRowsToProcess = 0;
+					int32 FillColorRow = -1;
+					for (int32 Y = 0; Y < TextureHeight; ++Y)
 					{
-						FillRowColorPixels(FillColorRow, Y);
+						if (!ProcessHorizontalRow(Y))
+						{
+							if (FillColorRow != -1)
+							{
+								FillRowColorPixels(FillColorRow, Y);
+							}
+							else
+							{
+								NumZeroedTopRowsToProcess = Y;
+							}
+						}
+						else
+						{
+							FillColorRow = Y;
+						}
 					}
-					else
+
+					// Can only fill upwards if image not fully zeroed
+					if (NumZeroedTopRowsToProcess > 0 && NumZeroedTopRowsToProcess + 1 < TextureHeight)
 					{
-						NumZeroedTopRowsToProcess = Y;
+						for (int32 Y = 0; Y <= NumZeroedTopRowsToProcess; ++Y)
+						{
+							FillRowColorPixels(NumZeroedTopRowsToProcess + 1, Y);
+						}
 					}
 				}
-				else
-				{
-					FillColorRow = Y;
-				}
-			}
 
-			// Can only fill upwards if image not fully zeroed
-			if (NumZeroedTopRowsToProcess > 0 && NumZeroedTopRowsToProcess + 1 < TextureHeight)
-			{
-				for (int32 Y = 0; Y <= NumZeroedTopRowsToProcess; ++Y)
+				/** returns False if requires further processing because entire row is filled with zeroed alpha values */
+				bool ProcessHorizontalRow(int32 Y)
 				{
-					FillRowColorPixels(NumZeroedTopRowsToProcess + 1, Y);
-				}
-			}
-		}
+					// only wipe out colors that are affected by png turning valid colors white if alpha = 0
+					const uint32 WhiteWithZeroAlpha = FColor(255, 255, 255, 0).DWColor();
 
-		/** returns False if requires further processing because entire row is filled with zeroed alpha values */
-		bool ProcessHorizontalRow(int32 Y)
-		{
-			// only wipe out colors that are affected by png turning valid colors white if alpha = 0
-			const uint32 WhiteWithZeroAlpha = FColor(255, 255, 255, 0).DWColor();
+					// Left -> Right
+					int32 NumLeftmostZerosToProcess = 0;
+					const PixelDataType* FillColor = nullptr;
+					for (int32 X = 0; X < TextureWidth; ++X)
+					{
+						PixelDataType* PixelData = SourceData + (Y * TextureWidth + X) * 4;
+						ColorDataType* ColorData = reinterpret_cast<ColorDataType*>(PixelData);
+
+						if (*ColorData == WhiteWithZeroAlpha)
+						{
+							if (FillColor)
+							{
+								PixelData[RIdx] = FillColor[RIdx];
+								PixelData[GIdx] = FillColor[GIdx];
+								PixelData[BIdx] = FillColor[BIdx];
+							}
+							else
+							{
+								// Mark pixel as needing fill
+								*ColorData = 0;
+
+								// Keep track of how many pixels to fill starting at beginning of row
+								NumLeftmostZerosToProcess = X;
+							}
+						}
+						else
+						{
+							FillColor = PixelData;
+						}
+					}
 
-			// Left -> Right
-			int32 NumLeftmostZerosToProcess = 0;
-			const PixelDataType* FillColor = nullptr;
-			for (int32 X = 0; X < TextureWidth; ++X)
-			{
-				PixelDataType* PixelData = SourceData + (Y * TextureWidth + X) * 4;
-				ColorDataType* ColorData = reinterpret_cast<ColorDataType*>(PixelData);
+					if (NumLeftmostZerosToProcess == 0)
+					{
+						// No pixels left that are zero
+						return true;
+					}
 
-				if (*ColorData == WhiteWithZeroAlpha)
-				{
-					if (FillColor)
+					if (NumLeftmostZerosToProcess + 1 >= TextureWidth)
+					{
+						// All pixels in this row are zero and must be filled using rows above or below
+						return false;
+					}
+
+					// Fill using non zero pixel immediately to the right of the beginning series of zeros
+					FillColor = SourceData + (Y * TextureWidth + NumLeftmostZerosToProcess + 1) * 4;
+
+					// Fill zero pixels found at beginning of row that could not be filled during the Left to Right pass
+					for (int32 X = 0; X <= NumLeftmostZerosToProcess; ++X)
 					{
+						PixelDataType* PixelData = SourceData + (Y * TextureWidth + X) * 4;
 						PixelData[RIdx] = FillColor[RIdx];
 						PixelData[GIdx] = FillColor[GIdx];
 						PixelData[BIdx] = FillColor[BIdx];
 					}
-					else
-					{
-						// Mark pixel as needing fill
-						*ColorData = 0;
 
-						// Keep track of how many pixels to fill starting at beginning of row
-						NumLeftmostZerosToProcess = X;
-					}
+					return true;
 				}
-				else
+
+				void FillRowColorPixels(int32 FillColorRow, int32 Y)
 				{
-					FillColor = PixelData;
+					for (int32 X = 0; X < TextureWidth; ++X)
+					{
+						const PixelDataType* FillColor = SourceData + (FillColorRow * TextureWidth + X) * 4;
+						PixelDataType* PixelData = SourceData + (Y * TextureWidth + X) * 4;
+						PixelData[RIdx] = FillColor[RIdx];
+						PixelData[GIdx] = FillColor[GIdx];
+						PixelData[BIdx] = FillColor[BIdx];
+					}
 				}
-			}
 
-			if (NumLeftmostZerosToProcess == 0)
+				PixelDataType* SourceData;
+				int32 TextureWidth;
+				int32 TextureHeight;
+			};
+
+			/**
+			 * For PNG texture importing, this ensures that any pixels with an alpha value of zero have an RGB
+			 * assigned to them from a neighboring pixel which has non-zero alpha.
+			 * This is needed as PNG exporters tend to turn pixels that are RGBA = (x,x,x,0) to (1,1,1,0)
+			 * and this produces artifacts when drawing the texture with bilinear filtering.
+			 *
+			 * @param TextureSource - The source texture
+			 * @param SourceData - The source texture data
+			 */
+			void FillZeroAlphaPNGData(int32 SizeX, int32 SizeY, ETextureSourceFormat SourceFormat, uint8* SourceData)
 			{
-				// No pixels left that are zero
-				return true;
-			}
-
-			if (NumLeftmostZerosToProcess + 1 >= TextureWidth)
-			{
-				// All pixels in this row are zero and must be filled using rows above or below
-				return false;
-			}
-
-			// Fill using non zero pixel immediately to the right of the beginning series of zeros
-			FillColor = SourceData + (Y * TextureWidth + NumLeftmostZerosToProcess + 1) * 4;
-
-			// Fill zero pixels found at beginning of row that could not be filled during the Left to Right pass
-			for (int32 X = 0; X <= NumLeftmostZerosToProcess; ++X)
-			{
-				PixelDataType* PixelData = SourceData + (Y * TextureWidth + X) * 4;
-				PixelData[RIdx] = FillColor[RIdx];
-				PixelData[GIdx] = FillColor[GIdx];
-				PixelData[BIdx] = FillColor[BIdx];
-			}
-
-			return true;
-		}
+				switch (SourceFormat)
+				{
+					case TSF_BGRA8:
+					{
+						TPNGDataFill<uint8, uint32, 2, 1, 0, 3> PNGFill(SizeX, SizeY, SourceData);
+						PNGFill.ProcessData();
+						break;
+					}
 
-		void FillRowColorPixels(int32 FillColorRow, int32 Y)
-		{
-			for (int32 X = 0; X < TextureWidth; ++X)
-			{
-				const PixelDataType* FillColor = SourceData + (FillColorRow * TextureWidth + X) * 4;
-				PixelDataType* PixelData = SourceData + (Y * TextureWidth + X) * 4;
-				PixelData[RIdx] = FillColor[RIdx];
-				PixelData[GIdx] = FillColor[GIdx];
-				PixelData[BIdx] = FillColor[BIdx];
+					case TSF_RGBA16:
+					{
+						TPNGDataFill<uint16, uint64, 0, 1, 2, 3> PNGFill(SizeX, SizeY, SourceData);
+						PNGFill.ProcessData();
+						break;
+					}
+				}
 			}
-		}
-
-		PixelDataType* SourceData;
-		int32 TextureWidth;
-		int32 TextureHeight;
-	};
-
-	/**
-	 * For PNG texture importing, this ensures that any pixels with an alpha value of zero have an RGB
-	 * assigned to them from a neighboring pixel which has non-zero alpha.
-	 * This is needed as PNG exporters tend to turn pixels that are RGBA = (x,x,x,0) to (1,1,1,0)
-	 * and this produces artifacts when drawing the texture with bilinear filtering.
-	 *
-	 * @param TextureSource - The source texture
-	 * @param SourceData - The source texture data
-	 */
-	void FillZeroAlphaPNGData(int32 SizeX, int32 SizeY, ETextureSourceFormat SourceFormat, uint8* SourceData)
-	{
-		switch (SourceFormat)
-		{
-		case TSF_BGRA8:
-		{
-			TPNGDataFill<uint8, uint32, 2, 1, 0, 3> PNGFill(SizeX, SizeY, SourceData);
-			PNGFill.ProcessData();
-			break;
-		}
-
-		case TSF_RGBA16:
-		{
-			TPNGDataFill<uint16, uint64, 0, 1, 2, 3> PNGFill(SizeX, SizeY, SourceData);
-			PNGFill.ProcessData();
-			break;
-		}
-		}
-	}
 
-} //PNGParserHelper namespace end
+		}//ns Private
+	}//ns Interchange
+}//ns UE
 
 bool UInterchangePNGTranslator::CanImportSourceData(const UInterchangeSourceData* InSourceData) const
 {
@@ -191,15 +196,15 @@ bool UInterchangePNGTranslator::CanImportSourceData(const UInterchangeSourceData
 
 bool UInterchangePNGTranslator::Translate(const UInterchangeSourceData* SourceData, UInterchangeBaseNodeContainer& BaseNodeContainer) const
 {
-	return Interchange::FTextureTranslatorUtilities::Generic2DTextureTranslate(SourceData, BaseNodeContainer);
+	return UE::Interchange::FTextureTranslatorUtilities::Generic2DTextureTranslate(SourceData, BaseNodeContainer);
 }
 
-TOptional<Interchange::FImportImage> UInterchangePNGTranslator::GetTexturePayloadData(const UInterchangeSourceData* SourceData, const FString& PayLoadKey) const
+TOptional<UE::Interchange::FImportImage> UInterchangePNGTranslator::GetTexturePayloadData(const UInterchangeSourceData* SourceData, const FString& PayLoadKey) const
 {
 	if (!SourceData)
 	{
 		UE_LOG(LogInterchangeImportPlugin, Error, TEXT("Failed to import PNG, bad source data."));
-		return TOptional<Interchange::FImportImage>();
+		return TOptional<UE::Interchange::FImportImage>();
 	}
 
 	TArray64<uint8> SourceDataBuffer;
@@ -209,19 +214,19 @@ TOptional<Interchange::FImportImage> UInterchangePNGTranslator::GetTexturePayloa
 	if (!Filename.Equals(PayLoadKey))
 	{
 		UE_LOG(LogInterchangeImportPlugin, Error, TEXT("Failed to import PNG, wrong payload key. [%s]"), *Filename);
-		return TOptional<Interchange::FImportImage>();
+		return TOptional<UE::Interchange::FImportImage>();
 	}
 
 	if (!FPaths::FileExists(Filename))
 	{
 		UE_LOG(LogInterchangeImportPlugin, Error, TEXT("Failed to import PNG, cannot open file. [%s]"), *Filename);
-		return TOptional<Interchange::FImportImage>();
+		return TOptional<UE::Interchange::FImportImage>();
 	}
 
 	if (!FFileHelper::LoadFileToArray(SourceDataBuffer, *Filename))
 	{
 		UE_LOG(LogInterchangeImportPlugin, Error, TEXT("Failed to import PNG, cannot load file content into an array. [%s]"), *Filename);
-		return TOptional<Interchange::FImportImage>();
+		return TOptional<UE::Interchange::FImportImage>();
 	}
 
 	const uint8* Buffer = SourceDataBuffer.GetData();
@@ -242,12 +247,12 @@ TOptional<Interchange::FImportImage> UInterchangePNGTranslator::GetTexturePayloa
 	if (!PngImageWrapper.IsValid() || !PngImageWrapper->SetCompressed(Buffer, Length))
 	{
 		UE_LOG(LogInterchangeImportPlugin, Error, TEXT("Failed to decode PNG. [%s]"), *Filename);
-		return TOptional<Interchange::FImportImage>();
+		return TOptional<UE::Interchange::FImportImage>();
 	}
-	if (!Interchange::FImportImageHelper::IsImportResolutionValid(PngImageWrapper->GetWidth(), PngImageWrapper->GetHeight(), bAllowNonPowerOfTwo))
+	if (!UE::Interchange::FImportImageHelper::IsImportResolutionValid(PngImageWrapper->GetWidth(), PngImageWrapper->GetHeight(), bAllowNonPowerOfTwo))
 	{
 		UE_LOG(LogInterchangeImportPlugin, Error, TEXT("Failed to import PNG, invalid resolution. Resolution[%d, %d], AllowPowerOfTwo[%s], [%s]"), PngImageWrapper->GetWidth(), PngImageWrapper->GetHeight(), bAllowNonPowerOfTwo ? TEXT("True") : TEXT("false"), *Filename);
-		return TOptional<Interchange::FImportImage>();
+		return TOptional<UE::Interchange::FImportImage>();
 	}
 
 	// Select the texture's source format
@@ -290,10 +295,10 @@ TOptional<Interchange::FImportImage> UInterchangePNGTranslator::GetTexturePayloa
 	if (TextureFormat == TSF_Invalid)
 	{
 		UE_LOG(LogInterchangeImportPlugin, Error, TEXT("PNG file [%s] contains data in an unsupported format"), *Filename);
-		return TOptional<Interchange::FImportImage>();
+		return TOptional<UE::Interchange::FImportImage>();
 	}
 
-	Interchange::FImportImage PayloadData;
+	UE::Interchange::FImportImage PayloadData;
 	PayloadData.Init2DWithParams(
 		PngImageWrapper->GetWidth(),
 		PngImageWrapper->GetHeight(),
@@ -309,13 +314,13 @@ TOptional<Interchange::FImportImage> UInterchangePNGTranslator::GetTexturePayloa
 		if (bFillPNGZeroAlpha)
 		{
 			// Replace the pixels with 0.0 alpha with a color value from the nearest neighboring color which has a non-zero alpha
-			PNGParserHelper::FillZeroAlphaPNGData(PayloadData.SizeX, PayloadData.SizeY, PayloadData.Format, PayloadData.RawData.GetData());
+			UE::Interchange::Private::FillZeroAlphaPNGData(PayloadData.SizeX, PayloadData.SizeY, PayloadData.Format, PayloadData.RawData.GetData());
 		}
 	}
 	else
 	{
 		UE_LOG(LogInterchangeImportPlugin, Error, TEXT("Failed to decode PNG. [%s]"), *Filename);
-		return TOptional<Interchange::FImportImage>();
+		return TOptional<UE::Interchange::FImportImage>();
 	}
 
 	return PayloadData;
diff --git a/Engine/Plugins/Runtime/Interchange/Source/Import/Private/Texture/InterchangePSDTranslator.cpp b/Engine/Plugins/Runtime/Interchange/Source/Import/Private/Texture/InterchangePSDTranslator.cpp
--- a/Engine/Plugins/Runtime/Interchange/Source/Import/Private/Texture/InterchangePSDTranslator.cpp
+++ b/Engine/Plugins/Runtime/Interchange/Source/Import/Private/Texture/InterchangePSDTranslator.cpp
@@ -17,350 +17,355 @@
 
 //////////////////////////////////////////////////////////////////////////
 // PSD helper local function
-
-namespace PSDParserHelper
+namespace UE
 {
-	#pragma pack(push,1)
-	struct FPSDFileHeader
-	{                                                           
-		int32	Signature;		// 8BPS
-		int16	Version;		// Version
-		int16	nChannels;		// Number of Channels (3=RGB) (4=RGBA)
-		int32	Height;			// Number of Image Rows
-		int32	Width;			// Number of Image Columns
-		int16	Depth;			// Number of Bits per Channel
-		int16	Mode;			// Image Mode (0=Bitmap)(1=Grayscale)(2=Indexed)(3=RGB)(4=CYMK)(7=Multichannel)
-		uint8	Pad[6];			// Padding
-
-		/**
-		 * @return Whether file has a valid signature
-		 */
-		bool IsValid() const
-		{
-			// Fail on bad signature
-			return (Signature == 0x38425053);
-		}
-
-		/**
-		 * @return Whether file has a supported version
-		 */
-		bool IsSupported() const
+	namespace Interchange
+	{
+		namespace Private
 		{
-			// Fail on bad version
-			if (Version != 1)
-				return false;
+#pragma pack(push,1)
+			struct FPSDFileHeader
+			{
+				int32	Signature;		// 8BPS
+				int16	Version;		// Version
+				int16	nChannels;		// Number of Channels (3=RGB) (4=RGBA)
+				int32	Height;			// Number of Image Rows
+				int32	Width;			// Number of Image Columns
+				int16	Depth;			// Number of Bits per Channel
+				int16	Mode;			// Image Mode (0=Bitmap)(1=Grayscale)(2=Indexed)(3=RGB)(4=CYMK)(7=Multichannel)
+				uint8	Pad[6];			// Padding
+
+				/**
+				 * @return Whether file has a valid signature
+				 */
+				bool IsValid() const
+				{
+					// Fail on bad signature
+					return (Signature == 0x38425053);
+				}
 
-			// Fail on anything other than 1, 3 or 4 channels
-			if ((nChannels != 1) && (nChannels != 3) && (nChannels != 4))
-				return false;
+				/**
+				 * @return Whether file has a supported version
+				 */
+				bool IsSupported() const
+				{
+					// Fail on bad version
+					if (Version != 1)
+						return false;
 
-			// Fail on anything other than 8 Bits/channel or 16 Bits/channel  
-			if ((Depth != 8) && (Depth != 16))
-				return false;
+					// Fail on anything other than 1, 3 or 4 channels
+					if ((nChannels != 1) && (nChannels != 3) && (nChannels != 4))
+						return false;
 
-			// Fail on anything other than Grayscale and RGB
-			// We can add support for indexed later if needed.
-			if (Mode != 1 && Mode != 3)
-				return false;
+					// Fail on anything other than 8 Bits/channel or 16 Bits/channel  
+					if ((Depth != 8) && (Depth != 16))
+						return false;
 
-			return true;
-		}
-	};
-	#pragma pack(pop)
+					// Fail on anything other than Grayscale and RGB
+					// We can add support for indexed later if needed.
+					if (Mode != 1 && Mode != 3)
+						return false;
 
+					return true;
+				}
+			};
+#pragma pack(pop)
 
-	static bool ReadData(uint8* pOut, const uint8*& pBuffer, FPSDFileHeader& Info)
-	{
-		const uint8* pPlane = nullptr;
-		const uint8* pRowTable = nullptr;
-		int32 iPlane;
-		int16 CompressionType;
-		int32 iPixel;
-		int32 iRow;
-		int32 CompressedBytes;
-		int32 iByte;
-		int32 Count;
-		uint8 Value;
-
-		// Double check to make sure this is a valid request
-		if (!Info.IsValid() || !Info.IsSupported())
-		{
-			return false;
-		}
-
-		const uint8* pCur = pBuffer + sizeof(FPSDFileHeader);
-		int32 NPixels = Info.Width * Info.Height;
-
-		int32 ClutSize = ((int32)pCur[0] << 24) +
-						 ((int32)pCur[1] << 16) +
-						 ((int32)pCur[2] <<  8) +
-						 ((int32)pCur[3] <<  0);
-		pCur += 4 + ClutSize;    
-
-		// Skip Image Resource Section
-		int32 ImageResourceSize = ((int32)pCur[0] << 24) +
-								  ((int32)pCur[1] << 16) +
-								  ((int32)pCur[2] <<  8) +
-								  ((int32)pCur[3] <<  0);
-		pCur += 4 + ImageResourceSize;
-
-		// Skip Layer and Mask Section
-		int32 LayerAndMaskSize =  ((int32)pCur[0] << 24) +
-								  ((int32)pCur[1] << 16) +
-								  ((int32)pCur[2] <<  8) +
-								  ((int32)pCur[3] <<  0);
-		pCur += 4 + LayerAndMaskSize;
-
-		// Determine number of bytes per pixel
-		int32 BytesPerPixel = 3;
-		const int32 BytesPerChannel = (Info.Depth / 8);
-		switch( Info.Mode )
-		{
-		case 1: // 'GrayScale'
-			BytesPerPixel = BytesPerChannel;
-			break;
-		case 2:        
-			BytesPerPixel = 1;        
-			return false;  // until we support indexed...
-			break;
-		case 3: // 'RGBColor'
-			if (Info.nChannels == 3)
-				BytesPerPixel = 3 * BytesPerChannel;
-			else                   
-				BytesPerPixel = 4 * BytesPerChannel;       
-			break;
-		default:
-			return false;
-			break;
-		}
-
-		// Get Compression Type
-		CompressionType = ((int32)pCur[0] <<  8) + ((int32)pCur[1] <<  0);    
-		pCur += 2;
-
-		// Fail on 16 Bits/channel with RLE. This can occur when the file is not saved with 'Maximize Compatibility'. Compression doesn't appear to be standard.
-		if(CompressionType == 1 && Info.Depth == 16)
-		{
-			return false;
-		}
 
-		// If no alpha channel, set alpha to opaque (255 or 65536).
-		if( Info.nChannels != 4)
-		{
-			if(Info.Depth == 8)
-			{
-				const uint32 Channels = 4;
-				const uint32 BufferSize = Info.Width * Info.Height * Channels * sizeof(uint8);
-				FMemory::Memset(pOut, 0xff, BufferSize);
-			}
-			else if(Info.Depth == 16)
+			static bool ReadData(uint8* pOut, const uint8*& pBuffer, FPSDFileHeader& Info)
 			{
-				const uint32 Channels = 4;
-				const uint32 BufferSize = Info.Width * Info.Height * Channels * sizeof(uint16);
-				FMemory::Memset(pOut, 0xff, BufferSize);
-			}
-		}
+				const uint8* pPlane = nullptr;
+				const uint8* pRowTable = nullptr;
+				int32 iPlane;
+				int16 CompressionType;
+				int32 iPixel;
+				int32 iRow;
+				int32 CompressedBytes;
+				int32 iByte;
+				int32 Count;
+				uint8 Value;
+
+				// Double check to make sure this is a valid request
+				if (!Info.IsValid() || !Info.IsSupported())
+				{
+					return false;
+				}
 
-		// Uncompressed?
-		if( CompressionType == 0 )
-		{
-			if(Info.Depth == 8)
-			{
-				FColor* Dest = (FColor*)pOut;
-				for(int32 Pixel=0; Pixel < NPixels; Pixel++ )
+				const uint8* pCur = pBuffer + sizeof(FPSDFileHeader);
+				int32 NPixels = Info.Width * Info.Height;
+
+				int32 ClutSize = ((int32)pCur[0] << 24) +
+					((int32)pCur[1] << 16) +
+					((int32)pCur[2] << 8) +
+					((int32)pCur[3] << 0);
+				pCur += 4 + ClutSize;
+
+				// Skip Image Resource Section
+				int32 ImageResourceSize = ((int32)pCur[0] << 24) +
+					((int32)pCur[1] << 16) +
+					((int32)pCur[2] << 8) +
+					((int32)pCur[3] << 0);
+				pCur += 4 + ImageResourceSize;
+
+				// Skip Layer and Mask Section
+				int32 LayerAndMaskSize = ((int32)pCur[0] << 24) +
+					((int32)pCur[1] << 16) +
+					((int32)pCur[2] << 8) +
+					((int32)pCur[3] << 0);
+				pCur += 4 + LayerAndMaskSize;
+
+				// Determine number of bytes per pixel
+				int32 BytesPerPixel = 3;
+				const int32 BytesPerChannel = (Info.Depth / 8);
+				switch (Info.Mode)
 				{
-					if (Info.nChannels == 1)
-					{
-						Dest[Pixel].R = pCur[NPixels + Pixel];
-						Dest[Pixel].G = pCur[NPixels + Pixel];
-						Dest[Pixel].B = pCur[NPixels + Pixel];
-					}
-					else
-					{
-						Dest[Pixel].R = pCur[NPixels * 0 + Pixel];
-						Dest[Pixel].G = pCur[NPixels * 1 + Pixel];
-						Dest[Pixel].B = pCur[NPixels * 2 + Pixel];
-						if (Info.nChannels == 4)
-						{
-							Dest[Pixel].A = pCur[NPixels * 3 + Pixel];
-						}
-					}
+					case 1: // 'GrayScale'
+						BytesPerPixel = BytesPerChannel;
+						break;
+					case 2:
+						BytesPerPixel = 1;
+						return false;  // until we support indexed...
+						break;
+					case 3: // 'RGBColor'
+						if (Info.nChannels == 3)
+							BytesPerPixel = 3 * BytesPerChannel;
+						else
+							BytesPerPixel = 4 * BytesPerChannel;
+						break;
+					default:
+						return false;
+						break;
 				}
-			}
-			else if (Info.Depth == 16)
-			{
-				uint32 SrcOffset = 0;
 
-				if (Info.nChannels == 1)
+				// Get Compression Type
+				CompressionType = ((int32)pCur[0] << 8) + ((int32)pCur[1] << 0);
+				pCur += 2;
+
+				// Fail on 16 Bits/channel with RLE. This can occur when the file is not saved with 'Maximize Compatibility'. Compression doesn't appear to be standard.
+				if (CompressionType == 1 && Info.Depth == 16)
 				{
-					uint16* Dest = (uint16*)pOut;
-					uint32 ChannelOffset = 0;
+					return false;
+				}
 
-					for (int32 Pixel = 0; Pixel < NPixels; Pixel++)
+				// If no alpha channel, set alpha to opaque (255 or 65536).
+				if (Info.nChannels != 4)
+				{
+					if (Info.Depth == 8)
 					{
-						Dest[ChannelOffset + 0] = ((pCur[SrcOffset] << 8) + (pCur[SrcOffset + 1] << 0));
-						Dest[ChannelOffset + 1] = ((pCur[SrcOffset] << 8) + (pCur[SrcOffset + 1] << 0));
-						Dest[ChannelOffset + 2] = ((pCur[SrcOffset] << 8) + (pCur[SrcOffset + 1] << 0));
-
-						//Increment offsets
-						ChannelOffset += 4;
-						SrcOffset += BytesPerChannel;
+						const uint32 Channels = 4;
+						const uint32 BufferSize = Info.Width * Info.Height * Channels * sizeof(uint8);
+						FMemory::Memset(pOut, 0xff, BufferSize);
+					}
+					else if (Info.Depth == 16)
+					{
+						const uint32 Channels = 4;
+						const uint32 BufferSize = Info.Width * Info.Height * Channels * sizeof(uint16);
+						FMemory::Memset(pOut, 0xff, BufferSize);
 					}
 				}
-				else
+
+				// Uncompressed?
+				if (CompressionType == 0)
 				{
-					// Loop through the planes	
-					for (iPlane = 0; iPlane < Info.nChannels; iPlane++)
+					if (Info.Depth == 8)
 					{
-						uint16* Dest = (uint16*)pOut;
-						uint32 ChannelOffset = iPlane;
-
+						FColor* Dest = (FColor*)pOut;
 						for (int32 Pixel = 0; Pixel < NPixels; Pixel++)
 						{
-							Dest[ChannelOffset] = ((pCur[SrcOffset] << 8) + (pCur[SrcOffset + 1] << 0));
-
-							//Increment offsets
-							ChannelOffset += 4;
-							SrcOffset += BytesPerChannel;
+							if (Info.nChannels == 1)
+							{
+								Dest[Pixel].R = pCur[NPixels + Pixel];
+								Dest[Pixel].G = pCur[NPixels + Pixel];
+								Dest[Pixel].B = pCur[NPixels + Pixel];
+							}
+							else
+							{
+								Dest[Pixel].R = pCur[NPixels * 0 + Pixel];
+								Dest[Pixel].G = pCur[NPixels * 1 + Pixel];
+								Dest[Pixel].B = pCur[NPixels * 2 + Pixel];
+								if (Info.nChannels == 4)
+								{
+									Dest[Pixel].A = pCur[NPixels * 3 + Pixel];
+								}
+							}
 						}
 					}
-				}
-			}
-		}
-		// RLE?
-		else if( CompressionType == 1 )
-		{
-			// Setup RowTable
-			pRowTable = pCur;
-			pCur += Info.nChannels*Info.Height*2;
-
-			FColor* Dest = (FColor*)pOut;
+					else if (Info.Depth == 16)
+					{
+						uint32 SrcOffset = 0;
 
-			// Loop through the planes
-			for( iPlane=0 ; iPlane<Info.nChannels ; iPlane++ )
-			{
-				int32 iWritePlane = iPlane;
-				if( iWritePlane > BytesPerPixel-1 ) iWritePlane = BytesPerPixel-1;
+						if (Info.nChannels == 1)
+						{
+							uint16* Dest = (uint16*)pOut;
+							uint32 ChannelOffset = 0;
 
-				// Loop through the rows
-				for( iRow=0 ; iRow<Info.Height ; iRow++ )
-				{
-					// Load a row
-					CompressedBytes = (pRowTable[(iPlane*Info.Height+iRow)*2  ] << 8) +
-						(pRowTable[(iPlane*Info.Height+iRow)*2+1] << 0);
-
-					// Setup Plane
-					pPlane = pCur;
-					pCur += CompressedBytes;
-
-					// Decompress Row
-					iPixel = 0;
-					iByte = 0;
-					while( (iPixel < Info.Width) && (iByte < CompressedBytes) )
-					{
-						int8 code = (int8)pPlane[iByte++];
+							for (int32 Pixel = 0; Pixel < NPixels; Pixel++)
+							{
+								Dest[ChannelOffset + 0] = ((pCur[SrcOffset] << 8) + (pCur[SrcOffset + 1] << 0));
+								Dest[ChannelOffset + 1] = ((pCur[SrcOffset] << 8) + (pCur[SrcOffset + 1] << 0));
+								Dest[ChannelOffset + 2] = ((pCur[SrcOffset] << 8) + (pCur[SrcOffset + 1] << 0));
 
-						// Is it a repeat?
-						if( code < 0 )
+								//Increment offsets
+								ChannelOffset += 4;
+								SrcOffset += BytesPerChannel;
+							}
+						}
+						else
 						{
-							Count = -(int32)code + 1;
-							Value = pPlane[iByte++];
-							while( Count-- > 0 )
+							// Loop through the planes	
+							for (iPlane = 0; iPlane < Info.nChannels; iPlane++)
 							{
-								int32 idx = (iPixel) + (iRow*Info.Width);
-								if (Info.nChannels == 1)
-								{
-									Dest[idx].R = Value;
-									Dest[idx].G = Value;
-									Dest[idx].B = Value;
-								}
-								else
+								uint16* Dest = (uint16*)pOut;
+								uint32 ChannelOffset = iPlane;
+
+								for (int32 Pixel = 0; Pixel < NPixels; Pixel++)
 								{
-									switch (iWritePlane)
-									{
-									case 0: Dest[idx].R = Value; break;
-									case 1: Dest[idx].G = Value; break;
-									case 2: Dest[idx].B = Value; break;
-									case 3: Dest[idx].A = Value; break;
-									}
+									Dest[ChannelOffset] = ((pCur[SrcOffset] << 8) + (pCur[SrcOffset + 1] << 0));
+
+									//Increment offsets
+									ChannelOffset += 4;
+									SrcOffset += BytesPerChannel;
 								}
-								iPixel++;
 							}
 						}
-						// Must be a literal then
-						else
+					}
+				}
+				// RLE?
+				else if (CompressionType == 1)
+				{
+					// Setup RowTable
+					pRowTable = pCur;
+					pCur += Info.nChannels * Info.Height * 2;
+
+					FColor* Dest = (FColor*)pOut;
+
+					// Loop through the planes
+					for (iPlane = 0; iPlane < Info.nChannels; iPlane++)
+					{
+						int32 iWritePlane = iPlane;
+						if (iWritePlane > BytesPerPixel - 1) iWritePlane = BytesPerPixel - 1;
+
+						// Loop through the rows
+						for (iRow = 0; iRow < Info.Height; iRow++)
 						{
-							Count = (int32)code + 1;
-							while( Count-- > 0 )
+							// Load a row
+							CompressedBytes = (pRowTable[(iPlane * Info.Height + iRow) * 2] << 8) +
+								(pRowTable[(iPlane * Info.Height + iRow) * 2 + 1] << 0);
+
+							// Setup Plane
+							pPlane = pCur;
+							pCur += CompressedBytes;
+
+							// Decompress Row
+							iPixel = 0;
+							iByte = 0;
+							while ((iPixel < Info.Width) && (iByte < CompressedBytes))
 							{
-								Value = pPlane[iByte++];
-								int32 idx = (iPixel) + (iRow*Info.Width);
+								int8 code = (int8)pPlane[iByte++];
 
-								if (Info.nChannels == 1)
+								// Is it a repeat?
+								if (code < 0)
 								{
-									Dest[idx].R = Value;
-									Dest[idx].G = Value;
-									Dest[idx].B = Value;
+									Count = -(int32)code + 1;
+									Value = pPlane[iByte++];
+									while (Count-- > 0)
+									{
+										int32 idx = (iPixel)+(iRow * Info.Width);
+										if (Info.nChannels == 1)
+										{
+											Dest[idx].R = Value;
+											Dest[idx].G = Value;
+											Dest[idx].B = Value;
+										}
+										else
+										{
+											switch (iWritePlane)
+											{
+												case 0: Dest[idx].R = Value; break;
+												case 1: Dest[idx].G = Value; break;
+												case 2: Dest[idx].B = Value; break;
+												case 3: Dest[idx].A = Value; break;
+											}
+										}
+										iPixel++;
+									}
 								}
+								// Must be a literal then
 								else
 								{
-									switch (iWritePlane)
+									Count = (int32)code + 1;
+									while (Count-- > 0)
 									{
-									case 0: Dest[idx].R = Value; break;
-									case 1: Dest[idx].G = Value; break;
-									case 2: Dest[idx].B = Value; break;
-									case 3: Dest[idx].A = Value; break;
+										Value = pPlane[iByte++];
+										int32 idx = (iPixel)+(iRow * Info.Width);
+
+										if (Info.nChannels == 1)
+										{
+											Dest[idx].R = Value;
+											Dest[idx].G = Value;
+											Dest[idx].B = Value;
+										}
+										else
+										{
+											switch (iWritePlane)
+											{
+												case 0: Dest[idx].R = Value; break;
+												case 1: Dest[idx].G = Value; break;
+												case 2: Dest[idx].B = Value; break;
+												case 3: Dest[idx].A = Value; break;
+											}
+										}
+										iPixel++;
 									}
 								}
-								iPixel++;
 							}
+
+							// Confirm that we decoded the right number of bytes
+							check(iByte == CompressedBytes);
+							check(iPixel == Info.Width);
 						}
 					}
-
-					// Confirm that we decoded the right number of bytes
-					check( iByte  == CompressedBytes );
-					check( iPixel == Info.Width );
 				}
-			}
-		}
-		else
-			return false;
+				else
+					return false;
 
-		// Success!
-		return true;
-	}
+				// Success!
+				return true;
+			}
 
-	static void GetPSDHeader( const uint8* Buffer, FPSDFileHeader& Info )
-	{
-		Info.Signature      =   ((int32)Buffer[ 0] << 24) +
-								((int32)Buffer[ 1] << 16) +
-								((int32)Buffer[ 2] <<  8) +
-								((int32)Buffer[ 3] <<  0);
+			static void GetPSDHeader(const uint8* Buffer, FPSDFileHeader& Info)
+			{
+				Info.Signature = ((int32)Buffer[0] << 24) +
+					((int32)Buffer[1] << 16) +
+					((int32)Buffer[2] << 8) +
+					((int32)Buffer[3] << 0);
 
-		Info.Version        =   ((int32)Buffer[ 4] <<  8) +
-								((int32)Buffer[ 5] <<  0);
+				Info.Version = ((int32)Buffer[4] << 8) +
+					((int32)Buffer[5] << 0);
 
-		Info.nChannels      =   ((int32)Buffer[12] <<  8) +
-								((int32)Buffer[13] <<  0);
+				Info.nChannels = ((int32)Buffer[12] << 8) +
+					((int32)Buffer[13] << 0);
 
-		Info.Height         =   ((int32)Buffer[14] << 24) +
-								((int32)Buffer[15] << 16) +
-								((int32)Buffer[16] <<  8) +
-								((int32)Buffer[17] <<  0);
+				Info.Height = ((int32)Buffer[14] << 24) +
+					((int32)Buffer[15] << 16) +
+					((int32)Buffer[16] << 8) +
+					((int32)Buffer[17] << 0);
 
-		Info.Width          =   ((int32)Buffer[18] << 24) +
-								((int32)Buffer[19] << 16) +
-								((int32)Buffer[20] <<  8) +
-								((int32)Buffer[21] <<  0);
+				Info.Width = ((int32)Buffer[18] << 24) +
+					((int32)Buffer[19] << 16) +
+					((int32)Buffer[20] << 8) +
+					((int32)Buffer[21] << 0);
 
-		Info.Depth          =   ((int32)Buffer[22] <<  8) +
-								((int32)Buffer[23] <<  0);
+				Info.Depth = ((int32)Buffer[22] << 8) +
+					((int32)Buffer[23] << 0);
 
-		Info.Mode           =   ((int32)Buffer[24] <<  8) +
-								((int32)Buffer[25] <<  0);
-	}
+				Info.Mode = ((int32)Buffer[24] << 8) +
+					((int32)Buffer[25] << 0);
+			}
 
-} // PSDParserHelper namespace end
+		}//ns Private
+	}//ns Interchange
+}//ns UE
 
 
 bool UInterchangePSDTranslator::CanImportSourceData(const UInterchangeSourceData* InSourceData) const
@@ -372,15 +377,15 @@ bool UInterchangePSDTranslator::CanImportSourceData(const UInterchangeSourceData
 
 bool UInterchangePSDTranslator::Translate(const UInterchangeSourceData* SourceData, UInterchangeBaseNodeContainer& BaseNodeContainer) const
 {
-	return Interchange::FTextureTranslatorUtilities::Generic2DTextureTranslate(SourceData, BaseNodeContainer);
+	return UE::Interchange::FTextureTranslatorUtilities::Generic2DTextureTranslate(SourceData, BaseNodeContainer);
 }
 
-TOptional<Interchange::FImportImage> UInterchangePSDTranslator::GetTexturePayloadData(const UInterchangeSourceData* SourceData, const FString& PayLoadKey) const
+TOptional<UE::Interchange::FImportImage> UInterchangePSDTranslator::GetTexturePayloadData(const UInterchangeSourceData* SourceData, const FString& PayLoadKey) const
 {
 	if (!SourceData)
 	{
 		UE_LOG(LogInterchangeImportPlugin, Error, TEXT("Failed to import PSD, bad source data."));
-		return TOptional<Interchange::FImportImage>();
+		return TOptional<UE::Interchange::FImportImage>();
 	}
 
 	TArray64<uint8> SourceDataBuffer;
@@ -390,19 +395,19 @@ TOptional<Interchange::FImportImage> UInterchangePSDTranslator::GetTexturePayloa
 	if (!Filename.Equals(PayLoadKey))
 	{
 		UE_LOG(LogInterchangeImportPlugin, Error, TEXT("Failed to import PSD, wrong payload key. [%s]"), *Filename);
-		return TOptional<Interchange::FImportImage>();
+		return TOptional<UE::Interchange::FImportImage>();
 	}
 
 	if (!FPaths::FileExists(Filename))
 	{
 		UE_LOG(LogInterchangeImportPlugin, Error, TEXT("Failed to import PSD, cannot open file. [%s]"), *Filename);
-		return TOptional<Interchange::FImportImage>();
+		return TOptional<UE::Interchange::FImportImage>();
 	}
 
 	if (!FFileHelper::LoadFileToArray(SourceDataBuffer, *Filename))
 	{
 		UE_LOG(LogInterchangeImportPlugin, Error, TEXT("Failed to import PSD, cannot load file content into an array. [%s]"), *Filename);
-		return TOptional<Interchange::FImportImage>();
+		return TOptional<UE::Interchange::FImportImage>();
 	}
 
 	const uint8* Buffer = SourceDataBuffer.GetData();
@@ -417,29 +422,29 @@ TOptional<Interchange::FImportImage> UInterchangePSDTranslator::GetTexturePayloa
 	//
 	// PSD File
 	//
-	PSDParserHelper::FPSDFileHeader PSDHeader;
-	if (Length < sizeof(PSDParserHelper::FPSDFileHeader))
+	UE::Interchange::Private::FPSDFileHeader PSDHeader;
+	if (Length < sizeof(UE::Interchange::Private::FPSDFileHeader))
 	{
-		return TOptional<Interchange::FImportImage>();
+		return TOptional<UE::Interchange::FImportImage>();
 	}
 
-	PSDParserHelper::GetPSDHeader(Buffer, PSDHeader);
+	UE::Interchange::Private::GetPSDHeader(Buffer, PSDHeader);
 
 	if (!PSDHeader.IsValid())
 	{
-		return TOptional<Interchange::FImportImage>();
+		return TOptional<UE::Interchange::FImportImage>();
 	}
 
 	// Check the resolution of the imported texture to ensure validity
-	if (!Interchange::FImportImageHelper::IsImportResolutionValid(PSDHeader.Width, PSDHeader.Height, bAllowNonPowerOfTwo))
+	if (!UE::Interchange::FImportImageHelper::IsImportResolutionValid(PSDHeader.Width, PSDHeader.Height, bAllowNonPowerOfTwo))
 	{
-		return TOptional<Interchange::FImportImage>();
+		return TOptional<UE::Interchange::FImportImage>();
 	}
 
 	if (!PSDHeader.IsSupported())
 	{
 		UE_LOG(LogInterchangeImportPlugin, Error, TEXT("Format of this PSD is not supported. Only Grayscale and RGBColor PSD images are currently supported, in 8-bit or 16-bit."));
-		return TOptional<Interchange::FImportImage>();
+		return TOptional<UE::Interchange::FImportImage>();
 	}
 
 	// Select the texture's source format
@@ -456,10 +461,10 @@ TOptional<Interchange::FImportImage> UInterchangePSDTranslator::GetTexturePayloa
 	if (TextureFormat == TSF_Invalid)
 	{
 		UE_LOG(LogInterchangeImportPlugin, Error, TEXT("PSD file contains data in an unsupported format."));
-		return TOptional<Interchange::FImportImage>();
+		return TOptional<UE::Interchange::FImportImage>();
 	}
 
-	Interchange::FImportImage PayloadData;
+	UE::Interchange::FImportImage PayloadData;
 
 	// The PSD is supported. Load it up.        
 	PayloadData.Init2DWithOneMip(
@@ -468,10 +473,10 @@ TOptional<Interchange::FImportImage> UInterchangePSDTranslator::GetTexturePayloa
 		TextureFormat
 	);
 
-	if (!PSDParserHelper::ReadData(PayloadData.RawData.GetData(), Buffer, PSDHeader))
+	if (!UE::Interchange::Private::ReadData(PayloadData.RawData.GetData(), Buffer, PSDHeader))
 	{
 		UE_LOG(LogInterchangeImportPlugin, Error, TEXT("Failed to read this PSD"));
-		return TOptional<Interchange::FImportImage>();
+		return TOptional<UE::Interchange::FImportImage>();
 	}
 
 	return PayloadData;
diff --git a/Engine/Plugins/Runtime/Interchange/Source/Import/Private/Texture/InterchangeTGATranslator.cpp b/Engine/Plugins/Runtime/Interchange/Source/Import/Private/Texture/InterchangeTGATranslator.cpp
--- a/Engine/Plugins/Runtime/Interchange/Source/Import/Private/Texture/InterchangeTGATranslator.cpp
+++ b/Engine/Plugins/Runtime/Interchange/Source/Import/Private/Texture/InterchangeTGATranslator.cpp
@@ -15,7 +15,7 @@
 
 namespace TGATranslatorImpl
 {
-	bool DecompressTGA( const FTGAFileHeader* TGA, Interchange::FImportImage& OutImage )
+	bool DecompressTGA( const FTGAFileHeader* TGA, UE::Interchange::FImportImage& OutImage )
 	{
 		if (TGA->ColorMapType == 1 && TGA->ImageTypeCode == 1 && TGA->BitsPerPixel == 8)
 		{
@@ -81,10 +81,10 @@ bool UInterchangeTGATranslator::CanImportSourceData(const UInterchangeSourceData
 
 bool UInterchangeTGATranslator::Translate(const UInterchangeSourceData* SourceData, UInterchangeBaseNodeContainer& BaseNodeContainer) const
 {
-	return Interchange::FTextureTranslatorUtilities::Generic2DTextureTranslate(SourceData, BaseNodeContainer);
+	return UE::Interchange::FTextureTranslatorUtilities::Generic2DTextureTranslate(SourceData, BaseNodeContainer);
 }
 
-TOptional<Interchange::FImportImage> UInterchangeTGATranslator::GetTexturePayloadData(const UInterchangeSourceData* SourceData, const FString& PayLoadKey) const
+TOptional<UE::Interchange::FImportImage> UInterchangeTGATranslator::GetTexturePayloadData(const UInterchangeSourceData* SourceData, const FString& PayLoadKey) const
 {
 	if (!SourceData)
 	{
@@ -128,14 +128,14 @@ TOptional<Interchange::FImportImage> UInterchangeTGATranslator::GetTexturePayloa
 	if ( TgaImageWrapper.IsValid() && TgaImageWrapper->SetCompressed( Buffer, Length ) )
 	{
 		// Check the resolution of the imported texture to ensure validity
-		if ( !Interchange::FImportImageHelper::IsImportResolutionValid( TgaImageWrapper->GetWidth(), TgaImageWrapper->GetHeight(), bAllowNonPowerOfTwo ) )
+		if ( !UE::Interchange::FImportImageHelper::IsImportResolutionValid( TgaImageWrapper->GetWidth(), TgaImageWrapper->GetHeight(), bAllowNonPowerOfTwo ) )
 		{
 			return {};
 		}
 
 		const FTGAFileHeader* TGA = (FTGAFileHeader *)Buffer;
 
-		Interchange::FImportImage PayloadData;
+		UE::Interchange::FImportImage PayloadData;
 
 		TGATranslatorImpl::DecompressTGA( TGA, PayloadData );
 
diff --git a/Engine/Plugins/Runtime/Interchange/Source/Import/Private/Texture/InterchangeTextureFactory.cpp b/Engine/Plugins/Runtime/Interchange/Source/Import/Private/Texture/InterchangeTextureFactory.cpp
--- a/Engine/Plugins/Runtime/Interchange/Source/Import/Private/Texture/InterchangeTextureFactory.cpp
+++ b/Engine/Plugins/Runtime/Interchange/Source/Import/Private/Texture/InterchangeTextureFactory.cpp
@@ -109,12 +109,12 @@ UObject* UInterchangeTextureFactory::CreateAsset(const UInterchangeTextureFactor
 	{
 		return nullptr;
 	}
-	const TOptional<Interchange::FImportImage> PayloadData = TextureTranslator->GetTexturePayloadData(Arguments.SourceData, PayLoadKey.GetValue());
+	const TOptional<UE::Interchange::FImportImage> PayloadData = TextureTranslator->GetTexturePayloadData(Arguments.SourceData, PayLoadKey.GetValue());
 	if(!PayloadData.IsSet())
 	{
 		return nullptr;
 	}
-	const Interchange::FImportImage& Image = PayloadData.GetValue();
+	const UE::Interchange::FImportImage& Image = PayloadData.GetValue();
 
 	const UClass* TextureClass = TextureNode->GetAssetClass();
 	check(TextureClass && TextureClass->IsChildOf(GetFactoryClass()));
diff --git a/Engine/Plugins/Runtime/Interchange/Source/Import/Private/Texture/InterchangeTexturePayloadData.cpp b/Engine/Plugins/Runtime/Interchange/Source/Import/Private/Texture/InterchangeTexturePayloadData.cpp
--- a/Engine/Plugins/Runtime/Interchange/Source/Import/Private/Texture/InterchangeTexturePayloadData.cpp
+++ b/Engine/Plugins/Runtime/Interchange/Source/Import/Private/Texture/InterchangeTexturePayloadData.cpp
@@ -8,7 +8,7 @@
 #include "Misc/MessageDialog.h"
 #include "RHI.h"
 
-bool Interchange::FImportImageHelper::IsImportResolutionValid(int32 Width, int32 Height, bool bAllowNonPowerOfTwo)
+bool UE::Interchange::FImportImageHelper::IsImportResolutionValid(int32 Width, int32 Height, bool bAllowNonPowerOfTwo)
 {
 	static const auto CVarVirtualTexturesEnabled = IConsoleManager::Get().FindTConsoleVariableDataInt(TEXT("r.VirtualTextures")); check(CVarVirtualTexturesEnabled);
 
@@ -49,7 +49,7 @@ bool Interchange::FImportImageHelper::IsImportResolutionValid(int32 Width, int32
 	return bValid;
 }
 
-void Interchange::FImportImage::Init2DWithParams(int32 InSizeX, int32 InSizeY, ETextureSourceFormat InFormat, bool InSRGB)
+void UE::Interchange::FImportImage::Init2DWithParams(int32 InSizeX, int32 InSizeY, ETextureSourceFormat InFormat, bool InSRGB)
 {
 	SizeX = InSizeX;
 	SizeY = InSizeY;
@@ -58,7 +58,7 @@ void Interchange::FImportImage::Init2DWithParams(int32 InSizeX, int32 InSizeY, E
 	SRGB = InSRGB;
 }
 
-void Interchange::FImportImage::Init2DWithOneMip(int32 InSizeX, int32 InSizeY, ETextureSourceFormat InFormat, const void* InData)
+void UE::Interchange::FImportImage::Init2DWithOneMip(int32 InSizeX, int32 InSizeY, ETextureSourceFormat InFormat, const void* InData)
 {
 #if WITH_EDITOR
 	SizeX = InSizeX;
@@ -73,7 +73,7 @@ void Interchange::FImportImage::Init2DWithOneMip(int32 InSizeX, int32 InSizeY, E
 #endif
 }
 
-void Interchange::FImportImage::Init2DWithMips(int32 InSizeX, int32 InSizeY, int32 InNumMips, ETextureSourceFormat InFormat, const void* InData)
+void UE::Interchange::FImportImage::Init2DWithMips(int32 InSizeX, int32 InSizeY, int32 InNumMips, ETextureSourceFormat InFormat, const void* InData)
 {
 	SizeX = InSizeX;
 	SizeY = InSizeY;
@@ -93,7 +93,7 @@ void Interchange::FImportImage::Init2DWithMips(int32 InSizeX, int32 InSizeY, int
 	}
 }
 
-int64 Interchange::FImportImage::GetMipSize(int32 InMipIndex) const
+int64 UE::Interchange::FImportImage::GetMipSize(int32 InMipIndex) const
 {
 #if WITH_EDITOR
 	check(InMipIndex >= 0);
@@ -106,7 +106,7 @@ int64 Interchange::FImportImage::GetMipSize(int32 InMipIndex) const
 #endif
 }
 
-void* Interchange::FImportImage::GetMipData(int32 InMipIndex)
+void* UE::Interchange::FImportImage::GetMipData(int32 InMipIndex)
 {
 	int64 Offset = 0;
 	for (int32 MipIndex = 0; MipIndex < InMipIndex; ++MipIndex)
diff --git a/Engine/Plugins/Runtime/Interchange/Source/Import/Private/Texture/TextureTranslatorUtilities.cpp b/Engine/Plugins/Runtime/Interchange/Source/Import/Private/Texture/TextureTranslatorUtilities.cpp
--- a/Engine/Plugins/Runtime/Interchange/Source/Import/Private/Texture/TextureTranslatorUtilities.cpp
+++ b/Engine/Plugins/Runtime/Interchange/Source/Import/Private/Texture/TextureTranslatorUtilities.cpp
@@ -9,7 +9,7 @@
 #include "UObject/Object.h"
 #include "UObject/ObjectMacros.h"
 
-bool Interchange::FTextureTranslatorUtilities::Generic2DTextureTranslate(const UInterchangeSourceData* SourceData, UInterchangeBaseNodeContainer& BaseNodeContainer)
+bool UE::Interchange::FTextureTranslatorUtilities::Generic2DTextureTranslate(const UInterchangeSourceData* SourceData, UInterchangeBaseNodeContainer& BaseNodeContainer)
 {
 	FString Filename = SourceData->GetFilename();
 	if (!FPaths::FileExists(Filename))
diff --git a/Engine/Plugins/Runtime/Interchange/Source/Import/Public/Fbx/InterchangeFbxTranslator.h b/Engine/Plugins/Runtime/Interchange/Source/Import/Public/Fbx/InterchangeFbxTranslator.h
--- a/Engine/Plugins/Runtime/Interchange/Source/Import/Public/Fbx/InterchangeFbxTranslator.h
+++ b/Engine/Plugins/Runtime/Interchange/Source/Import/Public/Fbx/InterchangeFbxTranslator.h
@@ -47,7 +47,7 @@ public:
 	 * @param PayloadKey - The key to retrieve the a particular payload contain into the specified source data.
 	 * @return a PayloadData containing the imported data. The TOptional will not be set if there is an error.
 	 */
-	virtual TOptional<Interchange::FImportImage> GetTexturePayloadData(const UInterchangeSourceData* SourceData, const FString& PayLoadKey) const override;
+	virtual TOptional<UE::Interchange::FImportImage> GetTexturePayloadData(const UInterchangeSourceData* SourceData, const FString& PayLoadKey) const override;
 
 	/* IInterchangeTexturePayloadInterface End */
 
@@ -63,7 +63,7 @@ public:
 	 * @param PayloadKey - The key to retrieve the a particular payload contain into the specified source data.
 	 * @return a PayloadData containing the imported data. The TOptional will not be set if there is an error.
 	 */
-	virtual TOptional<Interchange::FMaterialPayloadData> GetMaterialPayloadData(const UInterchangeSourceData* SourceData, const FString& PayLoadKey) const override;
+	virtual TOptional<UE::Interchange::FMaterialPayloadData> GetMaterialPayloadData(const UInterchangeSourceData* SourceData, const FString& PayLoadKey) const override;
 
 	/* IInterchangeMaterialPayloadInterface End */
 
@@ -79,7 +79,7 @@ public:
 	 * @param PayloadKey - The key to retrieve the a particular payload contain into the specified source data.
 	 * @return a PayloadData containing the imported data. The TOptional will not be set if there is an error.
 	 */
-	virtual TOptional<Interchange::FStaticMeshPayloadData> GetStaticMeshPayloadData(const UInterchangeSourceData* SourceData, const FString& PayLoadKey) const override;
+	virtual TOptional<UE::Interchange::FStaticMeshPayloadData> GetStaticMeshPayloadData(const UInterchangeSourceData* SourceData, const FString& PayLoadKey) const override;
 
 	/* IInterchangeStaticMeshPayloadInterface End */
 
@@ -95,7 +95,7 @@ public:
 	 * @param PayloadKey - The key to retrieve the a particular payload contain into the specified source data.
 	 * @return a PayloadData containing the imported data. The TOptional will not be set if there is an error.
 	 */
-	virtual TOptional<Interchange::FSkeletalMeshPayloadData> GetSkeletalMeshPayloadData(const UInterchangeSourceData* SourceData, const FString& PayLoadKey) const override;
+	virtual TOptional<UE::Interchange::FSkeletalMeshPayloadData> GetSkeletalMeshPayloadData(const UInterchangeSourceData* SourceData, const FString& PayLoadKey) const override;
 
 	/* IInterchangeSkeletalMeshPayloadInterface End */
 private:
@@ -106,7 +106,7 @@ private:
 	//We do not want to allocate the dispatcher and start the InterchangeWorker process
 	//in the constructor because Archetype, CDO and registered translators will
 	//never translate a source.
-	mutable TUniquePtr<InterchangeDispatcher::FInterchangeDispatcher> Dispatcher;
+	mutable TUniquePtr<UE::Interchange::FInterchangeDispatcher> Dispatcher;
 };
 
 
diff --git a/Engine/Plugins/Runtime/Interchange/Source/Import/Public/Material/InterchangeMaterialPayloadData.h b/Engine/Plugins/Runtime/Interchange/Source/Import/Public/Material/InterchangeMaterialPayloadData.h
--- a/Engine/Plugins/Runtime/Interchange/Source/Import/Public/Material/InterchangeMaterialPayloadData.h
+++ b/Engine/Plugins/Runtime/Interchange/Source/Import/Public/Material/InterchangeMaterialPayloadData.h
@@ -4,12 +4,15 @@
 
 #include "CoreMinimal.h"
 
-namespace Interchange
+namespace UE
 {
-	struct FMaterialPayloadData
+	namespace Interchange
 	{
-		//Store any material payload data here, currently not used
-	};
-}
+		struct FMaterialPayloadData
+		{
+			//Store any material payload data here, currently not used
+		};
+	}//ns Interchange
+}//ns UE
 
 
diff --git a/Engine/Plugins/Runtime/Interchange/Source/Import/Public/Material/InterchangeMaterialPayloadInterface.h b/Engine/Plugins/Runtime/Interchange/Source/Import/Public/Material/InterchangeMaterialPayloadInterface.h
--- a/Engine/Plugins/Runtime/Interchange/Source/Import/Public/Material/InterchangeMaterialPayloadInterface.h
+++ b/Engine/Plugins/Runtime/Interchange/Source/Import/Public/Material/InterchangeMaterialPayloadInterface.h
@@ -31,7 +31,7 @@ public:
 	 * @param PayloadKey - The key to retrieve the a particular payload contain into the specified source data.
 	 * @return a PayloadData containing the import image data. The TOptional will not be set if there is an error.
 	 */
-	virtual TOptional<Interchange::FMaterialPayloadData> GetMaterialPayloadData(const UInterchangeSourceData* SourceData, const FString& PayloadKey) const = 0;
+	virtual TOptional<UE::Interchange::FMaterialPayloadData> GetMaterialPayloadData(const UInterchangeSourceData* SourceData, const FString& PayloadKey) const = 0;
 };
 
 
diff --git a/Engine/Plugins/Runtime/Interchange/Source/Import/Public/Mesh/InterchangeSkeletalMeshPayload.h b/Engine/Plugins/Runtime/Interchange/Source/Import/Public/Mesh/InterchangeSkeletalMeshPayload.h
--- a/Engine/Plugins/Runtime/Interchange/Source/Import/Public/Mesh/InterchangeSkeletalMeshPayload.h
+++ b/Engine/Plugins/Runtime/Interchange/Source/Import/Public/Mesh/InterchangeSkeletalMeshPayload.h
@@ -5,13 +5,16 @@
 #include "CoreMinimal.h"
 #include "Rendering/SkeletalMeshLODImporterData.h"
 
-namespace Interchange
+namespace UE
 {
-	struct FSkeletalMeshPayloadData
+	namespace Interchange
 	{
-		//Currently the skeletalmesh payload data is editor only, we have to move to something available at runtime
+		struct FSkeletalMeshPayloadData
+		{
+			//Currently the skeletalmesh payload data is editor only, we have to move to something available at runtime
 #if WITH_EDITOR
-		TArray<FSkeletalMeshImportData> SkeletalMeshImportData;
+			TArray<FSkeletalMeshImportData> SkeletalMeshImportData;
 #endif
-	};
-}
+		};
+	}//ns Interchange
+}//ns UE
diff --git a/Engine/Plugins/Runtime/Interchange/Source/Import/Public/Mesh/InterchangeSkeletalMeshPayloadInterface.h b/Engine/Plugins/Runtime/Interchange/Source/Import/Public/Mesh/InterchangeSkeletalMeshPayloadInterface.h
--- a/Engine/Plugins/Runtime/Interchange/Source/Import/Public/Mesh/InterchangeSkeletalMeshPayloadInterface.h
+++ b/Engine/Plugins/Runtime/Interchange/Source/Import/Public/Mesh/InterchangeSkeletalMeshPayloadInterface.h
@@ -31,7 +31,7 @@ public:
 	 * @param PayloadKey - The key to retrieve the a particular payload contain into the specified source data.
 	 * @return a PayloadData containing the import image data. The TOptional will not be set if there is an error.
 	 */
-	virtual TOptional<Interchange::FSkeletalMeshPayloadData> GetSkeletalMeshPayloadData(const UInterchangeSourceData* SourceData, const FString& PayloadKey) const = 0;
+	virtual TOptional<UE::Interchange::FSkeletalMeshPayloadData> GetSkeletalMeshPayloadData(const UInterchangeSourceData* SourceData, const FString& PayloadKey) const = 0;
 };
 
 
diff --git a/Engine/Plugins/Runtime/Interchange/Source/Import/Public/Mesh/InterchangeStaticMeshPayload.h b/Engine/Plugins/Runtime/Interchange/Source/Import/Public/Mesh/InterchangeStaticMeshPayload.h
--- a/Engine/Plugins/Runtime/Interchange/Source/Import/Public/Mesh/InterchangeStaticMeshPayload.h
+++ b/Engine/Plugins/Runtime/Interchange/Source/Import/Public/Mesh/InterchangeStaticMeshPayload.h
@@ -5,10 +5,13 @@
 #include "CoreMinimal.h"
 #include "MeshDescription.h"
 
-namespace Interchange
+namespace UE
 {
-	struct FStaticMeshPayloadData
+	namespace Interchange
 	{
-		TArray<FMeshDescription> MeshDescriptions;
-	};
-}
+		struct FStaticMeshPayloadData
+		{
+			TArray<FMeshDescription> MeshDescriptions;
+		};
+	}//ns Interchange
+}//ns UE
diff --git a/Engine/Plugins/Runtime/Interchange/Source/Import/Public/Mesh/InterchangeStaticMeshPayloadInterface.h b/Engine/Plugins/Runtime/Interchange/Source/Import/Public/Mesh/InterchangeStaticMeshPayloadInterface.h
--- a/Engine/Plugins/Runtime/Interchange/Source/Import/Public/Mesh/InterchangeStaticMeshPayloadInterface.h
+++ b/Engine/Plugins/Runtime/Interchange/Source/Import/Public/Mesh/InterchangeStaticMeshPayloadInterface.h
@@ -31,7 +31,7 @@ public:
 	 * @param PayloadKey - The key to retrieve the a particular payload contain into the specified source data.
 	 * @return a PayloadData containing the import image data. The TOptional will not be set if there is an error.
 	 */
-	virtual TOptional<Interchange::FStaticMeshPayloadData> GetStaticMeshPayloadData(const UInterchangeSourceData* SourceData, const FString& PayloadKey) const = 0;
+	virtual TOptional<UE::Interchange::FStaticMeshPayloadData> GetStaticMeshPayloadData(const UInterchangeSourceData* SourceData, const FString& PayloadKey) const = 0;
 };
 
 
diff --git a/Engine/Plugins/Runtime/Interchange/Source/Import/Public/Texture/InterchangeBMPTranslator.h b/Engine/Plugins/Runtime/Interchange/Source/Import/Public/Texture/InterchangeBMPTranslator.h
--- a/Engine/Plugins/Runtime/Interchange/Source/Import/Public/Texture/InterchangeBMPTranslator.h
+++ b/Engine/Plugins/Runtime/Interchange/Source/Import/Public/Texture/InterchangeBMPTranslator.h
@@ -43,7 +43,7 @@ public:
 	 * @param PayloadKey - The key to retrieve the a particular payload contain into the specified source data.
 	 * @return a PayloadData containing the import image data. The TOptional will not be set if there is an error.
 	 */
-	virtual TOptional<Interchange::FImportImage> GetTexturePayloadData(const UInterchangeSourceData* SourceData, const FString& PayLoadKey) const override;
+	virtual TOptional<UE::Interchange::FImportImage> GetTexturePayloadData(const UInterchangeSourceData* SourceData, const FString& PayLoadKey) const override;
 
 	/* IInterchangeTexturePayloadInterface End */
 };
diff --git a/Engine/Plugins/Runtime/Interchange/Source/Import/Public/Texture/InterchangeDDSTranslator.h b/Engine/Plugins/Runtime/Interchange/Source/Import/Public/Texture/InterchangeDDSTranslator.h
--- a/Engine/Plugins/Runtime/Interchange/Source/Import/Public/Texture/InterchangeDDSTranslator.h
+++ b/Engine/Plugins/Runtime/Interchange/Source/Import/Public/Texture/InterchangeDDSTranslator.h
@@ -42,7 +42,7 @@ public:
 	 * @param PayloadKey - The key to retrieve the a particular payload contain into the specified source data.
 	 * @return a PayloadData containing the import image data. The TOptional will not be set if there is an error.
 	 */
-	virtual TOptional<Interchange::FImportImage> GetTexturePayloadData(const UInterchangeSourceData* SourceData, const FString& PayLoadKey) const override;
+	virtual TOptional<UE::Interchange::FImportImage> GetTexturePayloadData(const UInterchangeSourceData* SourceData, const FString& PayLoadKey) const override;
 
 	/* IInterchangeTexturePayloadInterface End */
 };
diff --git a/Engine/Plugins/Runtime/Interchange/Source/Import/Public/Texture/InterchangeTexturePayloadInterface.h b/Engine/Plugins/Runtime/Interchange/Source/Import/Public/Texture/InterchangeTexturePayloadInterface.h
--- a/Engine/Plugins/Runtime/Interchange/Source/Import/Public/Texture/InterchangeTexturePayloadInterface.h
+++ b/Engine/Plugins/Runtime/Interchange/Source/Import/Public/Texture/InterchangeTexturePayloadInterface.h
@@ -31,7 +31,7 @@ public:
 	 * @param PayloadKey - The key to retrieve the a particular payload contain into the specified source data.
 	 * @return a PayloadData containing the import image data. The TOptional will not be set if there is an error.
 	 */
-	virtual TOptional<Interchange::FImportImage> GetTexturePayloadData(const UInterchangeSourceData* SourceData, const FString& PayloadKey) const = 0;
+	virtual TOptional<UE::Interchange::FImportImage> GetTexturePayloadData(const UInterchangeSourceData* SourceData, const FString& PayloadKey) const = 0;
 };
 
 
diff --git a/Engine/Plugins/Runtime/Interchange/Source/Nodes/Public/InterchangeMaterialNode.h b/Engine/Plugins/Runtime/Interchange/Source/Nodes/Public/InterchangeMaterialNode.h
--- a/Engine/Plugins/Runtime/Interchange/Source/Nodes/Public/InterchangeMaterialNode.h
+++ b/Engine/Plugins/Runtime/Interchange/Source/Nodes/Public/InterchangeMaterialNode.h
@@ -16,18 +16,21 @@
 #include "InterchangeMaterialNode.generated.h"
 
 //Interchange namespace
-namespace Interchange
+namespace UE
 {
-	
-	struct FMaterialNodeStaticData : public FBaseNodeStaticData
+	namespace Interchange
 	{
-		static const FAttributeKey& PayloadSourceFileKey()
+
+		struct FMaterialNodeStaticData : public FBaseNodeStaticData
 		{
-			static FAttributeKey AttributeKey(TEXT("__PayloadSourceFile__"));
-			return AttributeKey;
-		}
-	};
-} //ns Interchange
+			static const FAttributeKey& PayloadSourceFileKey()
+			{
+				static FAttributeKey AttributeKey(TEXT("__PayloadSourceFile__"));
+				return AttributeKey;
+			}
+		};
+	} //ns Interchange
+}//ns UE
 
 /**
  * enum use to declare supported parameter
@@ -103,20 +106,20 @@ public:
 
 	virtual const TOptional<FString> GetPayLoadKey() const
 	{
-		if (!Attributes.ContainAttribute(Interchange::FMaterialNodeStaticData::PayloadSourceFileKey()))
+		if (!Attributes.ContainAttribute(UE::Interchange::FMaterialNodeStaticData::PayloadSourceFileKey()))
 		{
 			return TOptional<FString>();
 		}
-		Interchange::FAttributeStorage::TAttributeHandle<FString> AttributeHandle = Attributes.GetAttributeHandle<FString>(Interchange::FMaterialNodeStaticData::PayloadSourceFileKey());
+		UE::Interchange::FAttributeStorage::TAttributeHandle<FString> AttributeHandle = Attributes.GetAttributeHandle<FString>(UE::Interchange::FMaterialNodeStaticData::PayloadSourceFileKey());
 		if (!AttributeHandle.IsValid())
 		{
 			return TOptional<FString>();
 		}
 		FString PayloadKey;
-		Interchange::EAttributeStorageResult Result = AttributeHandle.Get(PayloadKey);
+		UE::Interchange::EAttributeStorageResult Result = AttributeHandle.Get(PayloadKey);
 		if (!IsAttributeStorageResultSuccess(Result))
 		{
-			LogAttributeStorageErrors(Result, TEXT("UInterchangeMaterialNode.GetPayLoadKey"), Interchange::FMaterialNodeStaticData::PayloadSourceFileKey());
+			LogAttributeStorageErrors(Result, TEXT("UInterchangeMaterialNode.GetPayLoadKey"), UE::Interchange::FMaterialNodeStaticData::PayloadSourceFileKey());
 			return TOptional<FString>();
 		}
 		return TOptional<FString>(PayloadKey);
@@ -125,10 +128,10 @@ public:
 	UFUNCTION(BlueprintCallable, Category = "Interchange | Node | Material")
 	virtual void SetPayLoadKey(const FString& PayloadKey)
 	{
-		Interchange::EAttributeStorageResult Result = Attributes.RegisterAttribute(Interchange::FMaterialNodeStaticData::PayloadSourceFileKey(), PayloadKey);
+		UE::Interchange::EAttributeStorageResult Result = Attributes.RegisterAttribute(UE::Interchange::FMaterialNodeStaticData::PayloadSourceFileKey(), PayloadKey);
 		if (!IsAttributeStorageResultSuccess(Result))
 		{
-			LogAttributeStorageErrors(Result, TEXT("UInterchangeMaterialNode.SetPayLoadKey"), Interchange::FMaterialNodeStaticData::PayloadSourceFileKey());
+			LogAttributeStorageErrors(Result, TEXT("UInterchangeMaterialNode.SetPayLoadKey"), UE::Interchange::FMaterialNodeStaticData::PayloadSourceFileKey());
 		}
 	}
 
@@ -360,11 +363,11 @@ private:
 	//This member is serialize manually in the serialize override
 	TMap<EInterchangeMaterialNodeParameterName, FParameterData> ParameterDatas;
 
-	const Interchange::FAttributeKey ClassNameAttributeKey = Interchange::FAttributeKey(TEXT("__ClassTypeAttribute__"));
+	const UE::Interchange::FAttributeKey ClassNameAttributeKey = UE::Interchange::FAttributeKey(TEXT("__ClassTypeAttribute__"));
 
 	//Material Attribute
-	const Interchange::FAttributeKey Macro_CustomMaterialUsageKey = Interchange::FAttributeKey(TEXT("MaterialUsage"));
-	const Interchange::FAttributeKey Macro_CustomBlendModeKey = Interchange::FAttributeKey(TEXT("BlendMode"));
+	const UE::Interchange::FAttributeKey Macro_CustomMaterialUsageKey = UE::Interchange::FAttributeKey(TEXT("MaterialUsage"));
+	const UE::Interchange::FAttributeKey Macro_CustomBlendModeKey = UE::Interchange::FAttributeKey(TEXT("BlendMode"));
 
 #if WITH_ENGINE
 	bool ApplyCustomMaterialUsageToAsset(UObject * Asset) const
diff --git a/Engine/Plugins/Runtime/Interchange/Source/Nodes/Public/InterchangeSceneNode.h b/Engine/Plugins/Runtime/Interchange/Source/Nodes/Public/InterchangeSceneNode.h
--- a/Engine/Plugins/Runtime/Interchange/Source/Nodes/Public/InterchangeSceneNode.h
+++ b/Engine/Plugins/Runtime/Interchange/Source/Nodes/Public/InterchangeSceneNode.h
@@ -66,6 +66,6 @@ public:
 
 private:
 	//Scene Attribute Keys
-	const Interchange::FAttributeKey Macro_CustomLocalTransformKey = Interchange::FAttributeKey(TEXT("LocalTransform"));
-	const Interchange::FAttributeKey Macro_CustomGlobalTransformKey = Interchange::FAttributeKey(TEXT("GlobalTransform"));
+	const UE::Interchange::FAttributeKey Macro_CustomLocalTransformKey = UE::Interchange::FAttributeKey(TEXT("LocalTransform"));
+	const UE::Interchange::FAttributeKey Macro_CustomGlobalTransformKey = UE::Interchange::FAttributeKey(TEXT("GlobalTransform"));
 };
