Make sure all interchange namespace respect the coding standard

diff --git a/Engine/Plugins/Runtime/Interchange/Source/Nodes/Public/InterchangeTextureNode.h b/Engine/Plugins/Runtime/Interchange/Source/Nodes/Public/InterchangeTextureNode.h
--- a/Engine/Plugins/Runtime/Interchange/Source/Nodes/Public/InterchangeTextureNode.h
+++ b/Engine/Plugins/Runtime/Interchange/Source/Nodes/Public/InterchangeTextureNode.h
@@ -13,19 +13,22 @@
 #include "InterchangeTextureNode.generated.h"
 
 //Interchange namespace
-namespace Interchange
+namespace UE
 {
-
-struct FTextureNodeStaticData : public FBaseNodeStaticData
-{
-	static const FAttributeKey& PayloadSourceFileKey()
+	namespace Interchange
 	{
-		static FAttributeKey AttributeKey(TEXT("__PayloadSourceFile__"));
-		return AttributeKey;
-	}
-};
 
-}
+		struct FTextureNodeStaticData : public FBaseNodeStaticData
+		{
+			static const FAttributeKey& PayloadSourceFileKey()
+			{
+				static FAttributeKey AttributeKey(TEXT("__PayloadSourceFile__"));
+				return AttributeKey;
+			}
+		};
+
+	}//ns Interchange
+}//ns UE
 
 UCLASS(BlueprintType)
 class INTERCHANGENODEPLUGIN_API UInterchangeTextureNode : public UInterchangeBaseNode
@@ -99,20 +102,20 @@ public:
 	/** Texture node Interface Begin */
 	virtual const TOptional<FString> GetPayLoadKey() const
 	{
-		if (!Attributes.ContainAttribute(Interchange::FTextureNodeStaticData::PayloadSourceFileKey()))
+		if (!Attributes.ContainAttribute(UE::Interchange::FTextureNodeStaticData::PayloadSourceFileKey()))
 		{
 			return TOptional<FString>();
 		}
-		Interchange::FAttributeStorage::TAttributeHandle<FString> AttributeHandle = Attributes.GetAttributeHandle<FString>(Interchange::FTextureNodeStaticData::PayloadSourceFileKey());
+		UE::Interchange::FAttributeStorage::TAttributeHandle<FString> AttributeHandle = Attributes.GetAttributeHandle<FString>(UE::Interchange::FTextureNodeStaticData::PayloadSourceFileKey());
 		if (!AttributeHandle.IsValid())
 		{
 			return TOptional<FString>();
 		}
 		FString PayloadKey;
-		Interchange::EAttributeStorageResult Result = AttributeHandle.Get(PayloadKey);
+		UE::Interchange::EAttributeStorageResult Result = AttributeHandle.Get(PayloadKey);
 		if (!IsAttributeStorageResultSuccess(Result))
 		{
-			LogAttributeStorageErrors(Result, TEXT("UInterchangeTextureNode.GetPayLoadKey"), Interchange::FTextureNodeStaticData::PayloadSourceFileKey());
+			LogAttributeStorageErrors(Result, TEXT("UInterchangeTextureNode.GetPayLoadKey"), UE::Interchange::FTextureNodeStaticData::PayloadSourceFileKey());
 			return TOptional<FString>();
 		}
 		return TOptional<FString>(PayloadKey);
@@ -121,10 +124,10 @@ public:
 	UFUNCTION(BlueprintCallable, Category = "Interchange | Node | Texture")
 	virtual void SetPayLoadKey(const FString& PayloadKey)
 	{
-		Interchange::EAttributeStorageResult Result = Attributes.RegisterAttribute(Interchange::FTextureNodeStaticData::PayloadSourceFileKey(), PayloadKey);
+		UE::Interchange::EAttributeStorageResult Result = Attributes.RegisterAttribute(UE::Interchange::FTextureNodeStaticData::PayloadSourceFileKey(), PayloadKey);
 		if (!IsAttributeStorageResultSuccess(Result))
 		{
-			LogAttributeStorageErrors(Result, TEXT("UInterchangeTextureNode.SetPayLoadKey"), Interchange::FTextureNodeStaticData::PayloadSourceFileKey());
+			LogAttributeStorageErrors(Result, TEXT("UInterchangeTextureNode.SetPayLoadKey"), UE::Interchange::FTextureNodeStaticData::PayloadSourceFileKey());
 		}
 	}
 
@@ -784,56 +787,56 @@ private:
 #endif
 	}
 
-	const Interchange::FAttributeKey ClassNameAttributeKey = Interchange::FAttributeKey(TEXT("__ClassTypeAttribute__"));
+	const UE::Interchange::FAttributeKey ClassNameAttributeKey = UE::Interchange::FAttributeKey(TEXT("__ClassTypeAttribute__"));
 
 	//Texture Adjustments
-	const Interchange::FAttributeKey Macro_CustomAdjustBrightnessKey = Interchange::FAttributeKey(TEXT("AdjustBrightness"));
-	const Interchange::FAttributeKey Macro_CustomAdjustBrightnessCurveKey = Interchange::FAttributeKey(TEXT("AdjustBrightnessCurve"));
-	const Interchange::FAttributeKey Macro_CustomAdjustVibranceKey = Interchange::FAttributeKey(TEXT("AdjustVibrance"));
-	const Interchange::FAttributeKey Macro_CustomAdjustSaturationKey = Interchange::FAttributeKey(TEXT("AdjustSaturation"));
-	const Interchange::FAttributeKey Macro_CustomAdjustRGBCurveKey = Interchange::FAttributeKey(TEXT("AdjustRGBCurve"));
-	const Interchange::FAttributeKey Macro_CustomAdjustHueKey = Interchange::FAttributeKey(TEXT("AdjustHue"));
-	const Interchange::FAttributeKey Macro_CustomAdjustMinAlphaKey = Interchange::FAttributeKey(TEXT("AdjustMinAlpha"));
-	const Interchange::FAttributeKey Macro_CustomAdjustMaxAlphaKey = Interchange::FAttributeKey(TEXT("AdjustMaxAlpha"));
-	const Interchange::FAttributeKey Macro_CustombChromaKeyTextureKey = Interchange::FAttributeKey(TEXT("bChromaKeyTexture"));
-	const Interchange::FAttributeKey Macro_CustomChromaKeyThresholdKey = Interchange::FAttributeKey(TEXT("ChromaKeyThreshold"));
-	const Interchange::FAttributeKey Macro_CustomChromaKeyColorKey = Interchange::FAttributeKey(TEXT("ChromaKeyColor"));
+	const UE::Interchange::FAttributeKey Macro_CustomAdjustBrightnessKey = UE::Interchange::FAttributeKey(TEXT("AdjustBrightness"));
+	const UE::Interchange::FAttributeKey Macro_CustomAdjustBrightnessCurveKey = UE::Interchange::FAttributeKey(TEXT("AdjustBrightnessCurve"));
+	const UE::Interchange::FAttributeKey Macro_CustomAdjustVibranceKey = UE::Interchange::FAttributeKey(TEXT("AdjustVibrance"));
+	const UE::Interchange::FAttributeKey Macro_CustomAdjustSaturationKey = UE::Interchange::FAttributeKey(TEXT("AdjustSaturation"));
+	const UE::Interchange::FAttributeKey Macro_CustomAdjustRGBCurveKey = UE::Interchange::FAttributeKey(TEXT("AdjustRGBCurve"));
+	const UE::Interchange::FAttributeKey Macro_CustomAdjustHueKey = UE::Interchange::FAttributeKey(TEXT("AdjustHue"));
+	const UE::Interchange::FAttributeKey Macro_CustomAdjustMinAlphaKey = UE::Interchange::FAttributeKey(TEXT("AdjustMinAlpha"));
+	const UE::Interchange::FAttributeKey Macro_CustomAdjustMaxAlphaKey = UE::Interchange::FAttributeKey(TEXT("AdjustMaxAlpha"));
+	const UE::Interchange::FAttributeKey Macro_CustombChromaKeyTextureKey = UE::Interchange::FAttributeKey(TEXT("bChromaKeyTexture"));
+	const UE::Interchange::FAttributeKey Macro_CustomChromaKeyThresholdKey = UE::Interchange::FAttributeKey(TEXT("ChromaKeyThreshold"));
+	const UE::Interchange::FAttributeKey Macro_CustomChromaKeyColorKey = UE::Interchange::FAttributeKey(TEXT("ChromaKeyColor"));
 
 	//Texture Compression
-	const Interchange::FAttributeKey Macro_CustomCompressionNoAlphaKey = Interchange::FAttributeKey(TEXT("CompressionNoAlpha"));
-	const Interchange::FAttributeKey Macro_CustomDeferCompressionKey = Interchange::FAttributeKey(TEXT("DeferCompression"));
-	const Interchange::FAttributeKey Macro_CustomLossyCompressionAmountKey = Interchange::FAttributeKey(TEXT("LossyCompressionAmount"));
-	const Interchange::FAttributeKey Macro_CustomMaxTextureSizeKey = Interchange::FAttributeKey(TEXT("MaxTextureSize"));
-	const Interchange::FAttributeKey Macro_CustomCompressionQualityKey = Interchange::FAttributeKey(TEXT("CompressionQuality"));
-	const Interchange::FAttributeKey Macro_CustomCompressionSettingsKey = Interchange::FAttributeKey(TEXT("CompressionSettings"));
+	const UE::Interchange::FAttributeKey Macro_CustomCompressionNoAlphaKey = UE::Interchange::FAttributeKey(TEXT("CompressionNoAlpha"));
+	const UE::Interchange::FAttributeKey Macro_CustomDeferCompressionKey = UE::Interchange::FAttributeKey(TEXT("DeferCompression"));
+	const UE::Interchange::FAttributeKey Macro_CustomLossyCompressionAmountKey = UE::Interchange::FAttributeKey(TEXT("LossyCompressionAmount"));
+	const UE::Interchange::FAttributeKey Macro_CustomMaxTextureSizeKey = UE::Interchange::FAttributeKey(TEXT("MaxTextureSize"));
+	const UE::Interchange::FAttributeKey Macro_CustomCompressionQualityKey = UE::Interchange::FAttributeKey(TEXT("CompressionQuality"));
+	const UE::Interchange::FAttributeKey Macro_CustomCompressionSettingsKey = UE::Interchange::FAttributeKey(TEXT("CompressionSettings"));
 
 	//Texture general
-	const Interchange::FAttributeKey Macro_CustombDitherMipMapAlphaKey = Interchange::FAttributeKey(TEXT("bDitherMipMapAlpha"));
-	const Interchange::FAttributeKey Macro_CustomAlphaCoverageThresholdsKey = Interchange::FAttributeKey(TEXT("AlphaCoverageThresholds"));
-	const Interchange::FAttributeKey Macro_CustombFlipGreenChannelKey = Interchange::FAttributeKey(TEXT("bFlipGreenChannel"));
-	const Interchange::FAttributeKey Macro_CustombForcePVRTC4Key = Interchange::FAttributeKey(TEXT("bForcePVRTC4"));
-	const Interchange::FAttributeKey Macro_CustomPowerOfTwoModeKey = Interchange::FAttributeKey(TEXT("PowerOfTwoMode"));
-	const Interchange::FAttributeKey Macro_CustomPaddingColorKey = Interchange::FAttributeKey(TEXT("PaddingColor"));
-	const Interchange::FAttributeKey Macro_CustomFilterKey = Interchange::FAttributeKey(TEXT("Filter"));
-	const Interchange::FAttributeKey Macro_CustomMipLoadOptionsKey = Interchange::FAttributeKey(TEXT("MipLoadOptions"));
-	const Interchange::FAttributeKey Macro_CustomSRGBKey = Interchange::FAttributeKey(TEXT("SRGB"));
-	const Interchange::FAttributeKey Macro_CustombUseLegacyGammaKey = Interchange::FAttributeKey(TEXT("bUseLegacyGamma"));
-	const Interchange::FAttributeKey Macro_CustomVirtualTextureStreamingKey = Interchange::FAttributeKey(TEXT("VirtualTextureStreaming"));
-	const Interchange::FAttributeKey Macro_CustomAddressXKey = Interchange::FAttributeKey(TEXT("AddressX"));
-	const Interchange::FAttributeKey Macro_CustomAddressYKey = Interchange::FAttributeKey(TEXT("AddressY"));
+	const UE::Interchange::FAttributeKey Macro_CustombDitherMipMapAlphaKey = UE::Interchange::FAttributeKey(TEXT("bDitherMipMapAlpha"));
+	const UE::Interchange::FAttributeKey Macro_CustomAlphaCoverageThresholdsKey = UE::Interchange::FAttributeKey(TEXT("AlphaCoverageThresholds"));
+	const UE::Interchange::FAttributeKey Macro_CustombFlipGreenChannelKey = UE::Interchange::FAttributeKey(TEXT("bFlipGreenChannel"));
+	const UE::Interchange::FAttributeKey Macro_CustombForcePVRTC4Key = UE::Interchange::FAttributeKey(TEXT("bForcePVRTC4"));
+	const UE::Interchange::FAttributeKey Macro_CustomPowerOfTwoModeKey = UE::Interchange::FAttributeKey(TEXT("PowerOfTwoMode"));
+	const UE::Interchange::FAttributeKey Macro_CustomPaddingColorKey = UE::Interchange::FAttributeKey(TEXT("PaddingColor"));
+	const UE::Interchange::FAttributeKey Macro_CustomFilterKey = UE::Interchange::FAttributeKey(TEXT("Filter"));
+	const UE::Interchange::FAttributeKey Macro_CustomMipLoadOptionsKey = UE::Interchange::FAttributeKey(TEXT("MipLoadOptions"));
+	const UE::Interchange::FAttributeKey Macro_CustomSRGBKey = UE::Interchange::FAttributeKey(TEXT("SRGB"));
+	const UE::Interchange::FAttributeKey Macro_CustombUseLegacyGammaKey = UE::Interchange::FAttributeKey(TEXT("bUseLegacyGamma"));
+	const UE::Interchange::FAttributeKey Macro_CustomVirtualTextureStreamingKey = UE::Interchange::FAttributeKey(TEXT("VirtualTextureStreaming"));
+	const UE::Interchange::FAttributeKey Macro_CustomAddressXKey = UE::Interchange::FAttributeKey(TEXT("AddressX"));
+	const UE::Interchange::FAttributeKey Macro_CustomAddressYKey = UE::Interchange::FAttributeKey(TEXT("AddressY"));
 
 	//Level of Detail
-	const Interchange::FAttributeKey Macro_CustombPreserveBorderKey = Interchange::FAttributeKey(TEXT("bPreserveBorder"));
-	const Interchange::FAttributeKey Macro_CustomMipGenSettingsKey = Interchange::FAttributeKey(TEXT("MipGenSettings"));
-	const Interchange::FAttributeKey Macro_CustomLODBiasKey = Interchange::FAttributeKey(TEXT("LODBias"));
-	const Interchange::FAttributeKey Macro_CustomLODGroupKey = Interchange::FAttributeKey(TEXT("LODGroup"));
+	const UE::Interchange::FAttributeKey Macro_CustombPreserveBorderKey = UE::Interchange::FAttributeKey(TEXT("bPreserveBorder"));
+	const UE::Interchange::FAttributeKey Macro_CustomMipGenSettingsKey = UE::Interchange::FAttributeKey(TEXT("MipGenSettings"));
+	const UE::Interchange::FAttributeKey Macro_CustomLODBiasKey = UE::Interchange::FAttributeKey(TEXT("LODBias"));
+	const UE::Interchange::FAttributeKey Macro_CustomLODGroupKey = UE::Interchange::FAttributeKey(TEXT("LODGroup"));
 	//TODO support per platform data in the FAttributeStorage, so we can set different value per platform at the pipeline stage, We set only the default value for now
-	const Interchange::FAttributeKey Macro_CustomDownscaleKey = Interchange::FAttributeKey(TEXT("Downscale"));
-	const Interchange::FAttributeKey Macro_CustomDownscaleOptionsKey = Interchange::FAttributeKey(TEXT("DownscaleOptions"));
+	const UE::Interchange::FAttributeKey Macro_CustomDownscaleKey = UE::Interchange::FAttributeKey(TEXT("Downscale"));
+	const UE::Interchange::FAttributeKey Macro_CustomDownscaleOptionsKey = UE::Interchange::FAttributeKey(TEXT("DownscaleOptions"));
 
 	//Compositing
-	const Interchange::FAttributeKey Macro_CustomCompositeTextureModeKey = Interchange::FAttributeKey(TEXT("CompositeTextureMode"));
-	const Interchange::FAttributeKey Macro_CustomCompositePowerKey = Interchange::FAttributeKey(TEXT("CompositePower"));
+	const UE::Interchange::FAttributeKey Macro_CustomCompositeTextureModeKey = UE::Interchange::FAttributeKey(TEXT("CompositeTextureMode"));
+	const UE::Interchange::FAttributeKey Macro_CustomCompositePowerKey = UE::Interchange::FAttributeKey(TEXT("CompositePower"));
 
 #if WITH_EDITORONLY_DATA
 	IMPLEMENT_NODE_ATTRIBUTE_APPLY_UOBJECT(AdjustBrightness, float, UTexture, );
diff --git a/Engine/Plugins/Runtime/Interchange/Source/Parsers/Fbx/Private/FbxAPI.cpp b/Engine/Plugins/Runtime/Interchange/Source/Parsers/Fbx/Private/FbxAPI.cpp
--- a/Engine/Plugins/Runtime/Interchange/Source/Parsers/Fbx/Private/FbxAPI.cpp
+++ b/Engine/Plugins/Runtime/Interchange/Source/Parsers/Fbx/Private/FbxAPI.cpp
@@ -13,7 +13,7 @@
 
 namespace UE
 {
-	namespace FbxParser
+	namespace Interchange
 	{
 		namespace Private
 		{
@@ -85,5 +85,5 @@ namespace UE
 				FFbxScene::AddHierarchy(SDKScene, NodeContainer, JSonErrorMessages);
 			}
 		} //ns Private
-	} //ns FbxParser
+	} //ns Interchange
 } //ns UE
diff --git a/Engine/Plugins/Runtime/Interchange/Source/Parsers/Fbx/Private/FbxAPI.h b/Engine/Plugins/Runtime/Interchange/Source/Parsers/Fbx/Private/FbxAPI.h
--- a/Engine/Plugins/Runtime/Interchange/Source/Parsers/Fbx/Private/FbxAPI.h
+++ b/Engine/Plugins/Runtime/Interchange/Source/Parsers/Fbx/Private/FbxAPI.h
@@ -12,7 +12,7 @@ class UInterchangeBaseNodeContainer;
 
 namespace UE
 {
-	namespace FbxParser
+	namespace Interchange
 	{
 		namespace Private
 		{
@@ -31,6 +31,6 @@ namespace UE
 				FbxImporter* SDKImporter = nullptr;
 				//FbxGeometryConverter* SDKGeometryConverter = nullptr;
 			};
-		}
-	}
-}
+		}//ns Private
+	}//ns Interchange
+}//ns UE
diff --git a/Engine/Plugins/Runtime/Interchange/Source/Parsers/Fbx/Private/FbxConvert.cpp b/Engine/Plugins/Runtime/Interchange/Source/Parsers/Fbx/Private/FbxConvert.cpp
--- a/Engine/Plugins/Runtime/Interchange/Source/Parsers/Fbx/Private/FbxConvert.cpp
+++ b/Engine/Plugins/Runtime/Interchange/Source/Parsers/Fbx/Private/FbxConvert.cpp
@@ -7,7 +7,7 @@
 
 namespace UE
 {
-	namespace FbxParser
+	namespace Interchange
 	{
 		namespace Private
 		{
@@ -174,6 +174,6 @@ namespace UE
 				UnrollFilter.Reset();
 				ApplyUnroll(SDKScene->GetRootNode(), Layer, &UnrollFilter);
 			}
-		}
-	}
-}
+		}//ns Private
+	}//ns Interchange
+}//ns UE
diff --git a/Engine/Plugins/Runtime/Interchange/Source/Parsers/Fbx/Private/FbxMaterial.cpp b/Engine/Plugins/Runtime/Interchange/Source/Parsers/Fbx/Private/FbxMaterial.cpp
--- a/Engine/Plugins/Runtime/Interchange/Source/Parsers/Fbx/Private/FbxMaterial.cpp
+++ b/Engine/Plugins/Runtime/Interchange/Source/Parsers/Fbx/Private/FbxMaterial.cpp
@@ -11,7 +11,7 @@
 
 namespace UE
 {
-	namespace FbxParser
+	namespace Interchange
 	{
 		namespace Private
 		{
@@ -183,5 +183,5 @@ namespace UE
 				}
 			}
 		} //ns Private
-	} //ns FbxParser
+	} //ns Interchange
 }//ns UE
diff --git a/Engine/Plugins/Runtime/Interchange/Source/Parsers/Fbx/Private/FbxScene.cpp b/Engine/Plugins/Runtime/Interchange/Source/Parsers/Fbx/Private/FbxScene.cpp
--- a/Engine/Plugins/Runtime/Interchange/Source/Parsers/Fbx/Private/FbxScene.cpp
+++ b/Engine/Plugins/Runtime/Interchange/Source/Parsers/Fbx/Private/FbxScene.cpp
@@ -11,11 +11,11 @@
 
 namespace UE
 {
-	namespace FbxParser
+	namespace Interchange
 	{
 		namespace Private
 		{
-			namespace ScenePrivate
+			namespace Scene
 			{
 				UInterchangeSceneNode* AddHierarchyRecursively(FbxNode* Node, FbxScene* SDKScene, UInterchangeBaseNodeContainer& NodeContainer, TArray<FString>& JSonErrorMessages)
 				{
@@ -88,7 +88,7 @@ namespace UE
 					}
 					return UnrealNode;
 				}
-			} //ns ScenePrivate
+			} //ns Scene
 
 			UInterchangeSceneNode* FFbxScene::CreateTransformNode(UInterchangeBaseNodeContainer& NodeContainer, const FString& NodeName, TArray<FString>& JSonErrorMessages)
 			{
@@ -109,8 +109,8 @@ namespace UE
 			void FFbxScene::AddHierarchy(FbxScene* SDKScene, UInterchangeBaseNodeContainer& NodeContainer, TArray<FString>& JSonErrorMessages)
 			{
 				 FbxNode* RootNode = SDKScene->GetRootNode();
-				 ScenePrivate::AddHierarchyRecursively(RootNode, SDKScene, NodeContainer, JSonErrorMessages);
+				 Scene::AddHierarchyRecursively(RootNode, SDKScene, NodeContainer, JSonErrorMessages);
 			}
 		} //ns Private
-	} //ns FbxParser
+	} //ns Interchange
 }//ns UE
diff --git a/Engine/Plugins/Runtime/Interchange/Source/Parsers/Fbx/Private/InterchangeFbxParser.cpp b/Engine/Plugins/Runtime/Interchange/Source/Parsers/Fbx/Private/InterchangeFbxParser.cpp
--- a/Engine/Plugins/Runtime/Interchange/Source/Parsers/Fbx/Private/InterchangeFbxParser.cpp
+++ b/Engine/Plugins/Runtime/Interchange/Source/Parsers/Fbx/Private/InterchangeFbxParser.cpp
@@ -10,29 +10,32 @@
 #include "UObject/UObjectGlobals.h"
 #include "UObject/Package.h"
 
-namespace InterchangeFbxParser
+namespace UE
 {
-	void FbxParser::LoadFbxFile(const FString& Filename, const FString& ResultFolder)
+	namespace Interchange
 	{
-		UE::FbxParser::Private::FFbxParser FbxParserPrivate;
-		if (!FbxParserPrivate.LoadFbxFile(Filename, JsonLoadMessages))
+		void FbxParser::LoadFbxFile(const FString& Filename, const FString& ResultFolder)
 		{
-			JsonLoadMessages.Add(TEXT("{\"Msg\" : {\"Type\" : \"Error\",\n\"Msg\" : \"Cannot load the fbx file.\"}}"));
-			return;
-		}
+			UE::Interchange::Private::FFbxParser FbxParserPrivate;
+			if (!FbxParserPrivate.LoadFbxFile(Filename, JsonLoadMessages))
+			{
+				JsonLoadMessages.Add(TEXT("{\"Msg\" : {\"Type\" : \"Error\",\n\"Msg\" : \"Cannot load the fbx file.\"}}"));
+				return;
+			}
 
-		
-		ResultFilepath = ResultFolder + TEXT("/SceneDescription.itc");
-		
-		Container = TStrongObjectPtr<UInterchangeBaseNodeContainer>(NewObject<UInterchangeBaseNodeContainer>(GetTransientPackage(), NAME_None));
-		if (!ensure(Container.IsValid()))
-		{
-			JsonLoadMessages.Add(TEXT("{\"Msg\" : {\"Type\" : \"Error\",\n\"Msg\" : \"Cannot allocate base node container to add fbx scene data\"}}"));
-			return;
+
+			ResultFilepath = ResultFolder + TEXT("/SceneDescription.itc");
+
+			Container = TStrongObjectPtr<UInterchangeBaseNodeContainer>(NewObject<UInterchangeBaseNodeContainer>(GetTransientPackage(), NAME_None));
+			if (!ensure(Container.IsValid()))
+			{
+				JsonLoadMessages.Add(TEXT("{\"Msg\" : {\"Type\" : \"Error\",\n\"Msg\" : \"Cannot allocate base node container to add fbx scene data\"}}"));
+				return;
+			}
+			FbxParserPrivate.FillContainerWithFbxScene(*Container.Get(), JsonLoadMessages);
+
+			Container.Get()->SaveToFile(ResultFilepath);
+			JsonLoadMessages.Add(TEXT("{\"Msg\" : {\"Type\" : \"Log\",\n\"Msg\" : \"This is a success!\"}}"));
 		}
-		FbxParserPrivate.FillContainerWithFbxScene(*Container.Get(), JsonLoadMessages);
-		
-		Container.Get()->SaveToFile(ResultFilepath);
-		JsonLoadMessages.Add(TEXT("{\"Msg\" : {\"Type\" : \"Log\",\n\"Msg\" : \"This is a success!\"}}"));
-	}
-}
+	}//ns Interchange
+}//ns UE
diff --git a/Engine/Plugins/Runtime/Interchange/Source/Parsers/Fbx/Public/InterchangeFbxParser.h b/Engine/Plugins/Runtime/Interchange/Source/Parsers/Fbx/Public/InterchangeFbxParser.h
--- a/Engine/Plugins/Runtime/Interchange/Source/Parsers/Fbx/Public/InterchangeFbxParser.h
+++ b/Engine/Plugins/Runtime/Interchange/Source/Parsers/Fbx/Public/InterchangeFbxParser.h
@@ -6,26 +6,27 @@
 #include "Nodes/InterchangeBaseNodeContainer.h"
 #include "UObject/StrongObjectPtr.h"
 
-namespace InterchangeFbxParser
+namespace UE
 {
+	namespace Interchange
+	{
+		class FbxParser
+		{
+		public:
+			/**
+			 * Parse a file support by the fbx sdk. It just extract all the fbx node and create a FBaseNodeContainer and dump it in a json file inside the ResultFolder
+			 * @param - Filename is the file that the fbx sdk will read (.fbx or .obj)
+			 * @param - ResultFolder is the folder where we must put any result file
+			 */
+			void LoadFbxFile(const FString& Filename, const FString& ResultFolder);
 
-class FbxParser
-{
-public:
-	/**
-	 * Parse a file support by the fbx sdk. It just extract all the fbx node and create a FBaseNodeContainer and dump it in a json file inside the ResultFolder
-	 * @param - Filename is the file that the fbx sdk will read (.fbx or .obj)
-	 * @param - ResultFolder is the folder where we must put any result file
-	 */
-	void LoadFbxFile(const FString& Filename, const FString& ResultFolder);
-	
-	FString GetResultFilepath(){ return ResultFilepath; }
-	TArray<FString> GetJsonLoadMessages() { return JsonLoadMessages; }
-
-private:
-	TStrongObjectPtr<UInterchangeBaseNodeContainer> Container = nullptr;
-	FString ResultFilepath;
-	TArray<FString> JsonLoadMessages;
-};
+			FString GetResultFilepath() { return ResultFilepath; }
+			TArray<FString> GetJsonLoadMessages() { return JsonLoadMessages; }
 
-} // ns InterchangeFbxParser
+		private:
+			TStrongObjectPtr<UInterchangeBaseNodeContainer> Container = nullptr;
+			FString ResultFilepath;
+			TArray<FString> JsonLoadMessages;
+		};
+	} // ns Interchange
+}//ns UE
diff --git a/Engine/Source/Developer/AssetTools/Private/AssetTools.cpp b/Engine/Source/Developer/AssetTools/Private/AssetTools.cpp
--- a/Engine/Source/Developer/AssetTools/Private/AssetTools.cpp
+++ b/Engine/Source/Developer/AssetTools/Private/AssetTools.cpp
@@ -1830,7 +1830,7 @@ TArray<UObject*> UAssetToolsImpl::ImportAssetsInternal(const TArray<FString>& Fi
 			const FString& Filename = FilesAndDestinations[FileIdx].Key;
 			const FString DestinationPath = ObjectTools::SanitizeObjectPath(FilesAndDestinations[FileIdx].Value);
 			{
-				Interchange::FScopedSourceData ScopedSourceData(Filename);
+				UE::Interchange::FScopedSourceData ScopedSourceData(Filename);
 				if (!InterchangeManager.CanTranslateSourceData(ScopedSourceData.GetSourceData()))
 				{
 					bOnlyInterchangeImport = false;
@@ -1864,7 +1864,7 @@ TArray<UObject*> UAssetToolsImpl::ImportAssetsInternal(const TArray<FString>& Fi
 		const FString DestinationPath = ObjectTools::SanitizeObjectPath(FilesAndDestinations[FileIdx].Value);
 		if (bUseInterchangeFramework)
 		{
-			Interchange::FScopedSourceData ScopedSourceData(Filename);
+			UE::Interchange::FScopedSourceData ScopedSourceData(Filename);
 			if (InterchangeManager.CanTranslateSourceData(ScopedSourceData.GetSourceData()))
 			{
 				auto PostImportedLambda = [bSyncToBrowser](UObject* ImportedObject)
diff --git a/Engine/Source/Editor/UnrealEd/Private/Editor.cpp b/Engine/Source/Editor/UnrealEd/Private/Editor.cpp
--- a/Engine/Source/Editor/UnrealEd/Private/Editor.cpp
+++ b/Engine/Source/Editor/UnrealEd/Private/Editor.cpp
@@ -383,7 +383,7 @@ bool FReimportManager::Reimport( UObject* Obj, bool bAskForNewFileIfMissing, boo
 			{
 				if (bUseInterchangeFramework)
 				{
-					Interchange::FScopedSourceData ScopedSourceData(SourceFilenames[0]);
+					UE::Interchange::FScopedSourceData ScopedSourceData(SourceFilenames[0]);
 					if (InterchangeManager.CanTranslateSourceData(ScopedSourceData.GetSourceData()))
 					{
 						auto PostImportedLambda = [](UObject* ImportedObject)
diff --git a/Engine/Source/Programs/InterchangeWorker/Private/InterchangeWorker.cpp b/Engine/Source/Programs/InterchangeWorker/Private/InterchangeWorker.cpp
--- a/Engine/Source/Programs/InterchangeWorker/Private/InterchangeWorker.cpp
+++ b/Engine/Source/Programs/InterchangeWorker/Private/InterchangeWorker.cpp
@@ -53,12 +53,12 @@ int32 Main(int32 Argc, TCHAR * Argv[])
 	int32 Major = 0;
 	int32 Minor = 0;
 	int32 Patch = 0;
-	if (!InterchangeDispatcher::DispatcherCommandVersion::FromString(InterchangeDispatcherVersion, Major, Minor, Patch))
+	if (!UE::Interchange::DispatcherCommandVersion::FromString(InterchangeDispatcherVersion, Major, Minor, Patch))
 	{
 		UE_LOG(LogInterchangeWorker, Error, TEXT("Incompatible interchange dispatcher version string command argument."));
 		return EXIT_FAILURE;
 	}
- 	if (!InterchangeDispatcher::DispatcherCommandVersion::IsAPICompatible(Major, Minor, Patch))
+ 	if (!UE::Interchange::DispatcherCommandVersion::IsAPICompatible(Major, Minor, Patch))
  	{
  		UE_LOG(LogInterchangeWorker, Error, TEXT("Incompatible interchange dispatcher version. Please recompile InterchangeWorker target."));
  		return EXIT_FAILURE;
diff --git a/Engine/Source/Programs/InterchangeWorker/Private/InterchangeWorkerImpl.cpp b/Engine/Source/Programs/InterchangeWorker/Private/InterchangeWorkerImpl.cpp
--- a/Engine/Source/Programs/InterchangeWorker/Private/InterchangeWorkerImpl.cpp
+++ b/Engine/Source/Programs/InterchangeWorker/Private/InterchangeWorkerImpl.cpp
@@ -17,8 +17,7 @@
 #include "SocketSubsystem.h"
 #include "Sockets.h"
 
-using namespace InterchangeDispatcher;
-
+using namespace UE::Interchange;
 
 FInterchangeWorkerImpl::FInterchangeWorkerImpl(int32 InServerPID, int32 InServerPort, FString& InResultFolder)
 	: ServerPID(InServerPID)
@@ -127,7 +126,7 @@ void FInterchangeWorkerImpl::ProcessCommand(const FRunTaskCommand& RunTaskComman
 	//Process the json and run the task
 	FString JSonResult;
 	TArray<FString> JSonMessages;
-	InterchangeDispatcher::FJsonLoadSourceCmd LoadSourceCommand;
+	FJsonLoadSourceCmd LoadSourceCommand;
 	//Any command FromJson function return true if the Json descibe the command
 	if (LoadSourceCommand.FromJson(JsonToProcess))
 	{
@@ -157,20 +156,20 @@ void FInterchangeWorkerImpl::ProcessCommand(const FRunTaskCommand& RunTaskComman
 	UE_LOG(LogInterchangeWorker, Verbose, TEXT("End of Process %s"), *JsonToProcess);
 }
 
-ETaskState FInterchangeWorkerImpl::LoadFbxFile(const InterchangeDispatcher::FJsonLoadSourceCmd& LoadSourceCommand, FString& OutJSonResult, TArray<FString>& OutJSonMessages) const
+ETaskState FInterchangeWorkerImpl::LoadFbxFile(const FJsonLoadSourceCmd& LoadSourceCommand, FString& OutJSonResult, TArray<FString>& OutJSonMessages) const
 {
 	ETaskState ResultState = ETaskState::Unknown;
 	FString SourceFilename = LoadSourceCommand.GetSourceFilename();
 #if PLATFORM_WINDOWS && PLATFORM_64BITS
-	InterchangeFbxParser::FbxParser FbxParser;
+	UE::Interchange::FbxParser FbxParser;
 	FbxParser.LoadFbxFile(SourceFilename, ResultFolder);
-	InterchangeDispatcher::FJsonLoadSourceCmd::JsonResultParser ResultParser;
+	FJsonLoadSourceCmd::JsonResultParser ResultParser;
 	ResultParser.SetResultFilename(FbxParser.GetResultFilepath());
 	OutJSonMessages = FbxParser.GetJsonLoadMessages();
 	OutJSonResult = ResultParser.ToJson();
 	ResultState = ETaskState::ProcessOk;
 #else
-	InterchangeDispatcher::FJsonLoadSourceCmd::JsonResultParser ResultParser;
+	FJsonLoadSourceCmd::JsonResultParser ResultParser;
 	ResultParser.SetResultFilename(FString());
 	OutJSonMessages.Add(TEXT("{\"Msg\" : {\"Type\" : \"Error\",\n\"Msg\" : \"Cannot Execute fbx command on other platform then window 64 bits!\"}}"));
 	ResultState = ETaskState::ProcessFailed;
diff --git a/Engine/Source/Programs/InterchangeWorker/Private/InterchangeWorkerImpl.h b/Engine/Source/Programs/InterchangeWorker/Private/InterchangeWorkerImpl.h
--- a/Engine/Source/Programs/InterchangeWorker/Private/InterchangeWorkerImpl.h
+++ b/Engine/Source/Programs/InterchangeWorker/Private/InterchangeWorkerImpl.h
@@ -20,15 +20,15 @@ public:
 
 private:
 	void InitiatePing();
-	void ProcessCommand(const InterchangeDispatcher::FPingCommand& PingCommand);
-	void ProcessCommand(const InterchangeDispatcher::FBackPingCommand& BackPingCommand);
-	void ProcessCommand(const InterchangeDispatcher::FRunTaskCommand& TerminateCommand);
+	void ProcessCommand(const UE::Interchange::FPingCommand& PingCommand);
+	void ProcessCommand(const UE::Interchange::FBackPingCommand& BackPingCommand);
+	void ProcessCommand(const UE::Interchange::FRunTaskCommand& TerminateCommand);
 
-	InterchangeDispatcher::ETaskState LoadFbxFile(const InterchangeDispatcher::FJsonLoadSourceCmd& LoadSourceCommand, FString& OutJSonResult, TArray<FString>& OutJSonMessages) const;
+	UE::Interchange::ETaskState LoadFbxFile(const UE::Interchange::FJsonLoadSourceCmd& LoadSourceCommand, FString& OutJSonResult, TArray<FString>& OutJSonMessages) const;
 
 private:
-	InterchangeDispatcher::FNetworkClientNode NetworkInterface;
-	InterchangeDispatcher::FCommandQueue CommandIO;
+	UE::Interchange::FNetworkClientNode NetworkInterface;
+	UE::Interchange::FCommandQueue CommandIO;
 
 	int32 ServerPID;
 	int32 ServerPort;
diff --git a/Engine/Source/Runtime/Interchange/Core/Private/InterchangeManager.cpp b/Engine/Source/Runtime/Interchange/Core/Private/InterchangeManager.cpp
--- a/Engine/Source/Runtime/Interchange/Core/Private/InterchangeManager.cpp
+++ b/Engine/Source/Runtime/Interchange/Core/Private/InterchangeManager.cpp
@@ -43,35 +43,35 @@ namespace InternalInterchangePrivate
 	}
 }
 
-Interchange::FScopedSourceData::FScopedSourceData(const FString& Filename)
+UE::Interchange::FScopedSourceData::FScopedSourceData(const FString& Filename)
 {
 	//Found the translator
 	SourceDataPtr = TStrongObjectPtr<UInterchangeSourceData>(UInterchangeManager::GetInterchangeManager().CreateSourceData(Filename));
 	check(SourceDataPtr.IsValid());
 }
 
-UInterchangeSourceData* Interchange::FScopedSourceData::GetSourceData() const
+UInterchangeSourceData* UE::Interchange::FScopedSourceData::GetSourceData() const
 {
 	return SourceDataPtr.Get();
 }
 
-Interchange::FScopedTranslator::FScopedTranslator(const UInterchangeSourceData* SourceData)
+UE::Interchange::FScopedTranslator::FScopedTranslator(const UInterchangeSourceData* SourceData)
 {
 	//Found the translator
 	ScopedTranslatorPtr = TStrongObjectPtr<UInterchangeTranslatorBase>(UInterchangeManager::GetInterchangeManager().GetTranslatorForSourceData(SourceData));
 }
 
-UInterchangeTranslatorBase* Interchange::FScopedTranslator::GetTranslator()
+UInterchangeTranslatorBase* UE::Interchange::FScopedTranslator::GetTranslator()
 {
 	return ScopedTranslatorPtr.Get();
 }
 
-Interchange::FImportAsyncHelper::FImportAsyncHelper()
+UE::Interchange::FImportAsyncHelper::FImportAsyncHelper()
 {
 	RootObjectCompletionEvent = FGraphEvent::CreateGraphEvent();
 }
 
-void Interchange::FImportAsyncHelper::AddReferencedObjects(FReferenceCollector& Collector)
+void UE::Interchange::FImportAsyncHelper::AddReferencedObjects(FReferenceCollector& Collector)
 {
 	for (UInterchangeSourceData* SourceData : SourceDatas)
 	{
@@ -91,7 +91,7 @@ void Interchange::FImportAsyncHelper::AddReferencedObjects(FReferenceCollector&
 	}
 }
 
-void Interchange::FImportAsyncHelper::CleanUp()
+void UE::Interchange::FImportAsyncHelper::CleanUp()
 {
 	//Release the graph
 	BaseNodeContainers.Empty();
@@ -134,18 +134,18 @@ void Interchange::FImportAsyncHelper::CleanUp()
 	Factories.Empty();
 }
 
-Interchange::FAsyncImportResult::FAsyncImportResult( TFuture< UObject* >&& InFutureObject, const FGraphEventRef& InGraphEvent )
+UE::Interchange::FAsyncImportResult::FAsyncImportResult( TFuture< UObject* >&& InFutureObject, const FGraphEventRef& InGraphEvent )
 	: FutureObject( MoveTemp( InFutureObject ) )
 	, GraphEvent( InGraphEvent )
 {
 }
 
-bool Interchange::FAsyncImportResult::IsValid() const
+bool UE::Interchange::FAsyncImportResult::IsValid() const
 {
 	return FutureObject.IsValid();
 }
 
-UObject* Interchange::FAsyncImportResult::Get() const
+UObject* UE::Interchange::FAsyncImportResult::Get() const
 {
 	if ( !FutureObject.IsReady() )
 	{
@@ -156,12 +156,12 @@ UObject* Interchange::FAsyncImportResult::Get() const
 	return FutureObject.Get();
 }
 
-Interchange::FAsyncImportResult Interchange::FAsyncImportResult::Next( TFunction< UObject*( UObject* ) > Continuation )
+UE::Interchange::FAsyncImportResult UE::Interchange::FAsyncImportResult::Next( TFunction< UObject*( UObject* ) > Continuation )
 {
-	return Interchange::FAsyncImportResult{ FutureObject.Next( Continuation ), GraphEvent };
+	return UE::Interchange::FAsyncImportResult{ FutureObject.Next( Continuation ), GraphEvent };
 }
 
-void Interchange::SanitizeInvalidChar(FString& String)
+void UE::Interchange::SanitizeInvalidChar(FString& String)
 {
 	const TCHAR* InvalidChar = INVALID_OBJECTPATH_CHARACTERS;
 	while (*InvalidChar)
@@ -235,7 +235,7 @@ bool UInterchangeManager::RegisterWriter(const UClass* WriterClass)
 bool UInterchangeManager::CanTranslateSourceData(const UInterchangeSourceData* SourceData) const
 {
 	//Found the translator
-	Interchange::FScopedTranslator ScopeDataTranslator(SourceData);
+	UE::Interchange::FScopedTranslator ScopeDataTranslator(SourceData);
 	const UInterchangeTranslatorBase* SourceDataTranslator = ScopeDataTranslator.GetTranslator();
 	if (SourceDataTranslator)
 	{
@@ -249,23 +249,23 @@ bool UInterchangeManager::ImportAsset(const FString& ContentPath, const UInterch
 	return ImportAssetAsync( ContentPath, SourceData, ImportAssetParameters ).IsValid();
 }
 
-Interchange::FAsyncImportResult UInterchangeManager::ImportAssetAsync(const FString& ContentPath, const UInterchangeSourceData* SourceData, const FImportAssetParameters& ImportAssetParameters)
+UE::Interchange::FAsyncImportResult UInterchangeManager::ImportAssetAsync(const FString& ContentPath, const UInterchangeSourceData* SourceData, const FImportAssetParameters& ImportAssetParameters)
 {
 	FString PackageBasePath = ContentPath;
 	if(!ImportAssetParameters.ReimportAsset)
 	{
-		Interchange::SanitizeInvalidChar(PackageBasePath);
+		UE::Interchange::SanitizeInvalidChar(PackageBasePath);
 	}
 
 	bool bCanShowDialog = !ImportAssetParameters.bIsAutomated && !IsAttended();
 
 	//Create a task for every source data
-	Interchange::FImportAsyncHelperData TaskData;
+	UE::Interchange::FImportAsyncHelperData TaskData;
 	TaskData.bIsAutomated = ImportAssetParameters.bIsAutomated;
-	TaskData.ImportType = Interchange::EImportType::ImportType_Asset;
+	TaskData.ImportType = UE::Interchange::EImportType::ImportType_Asset;
 	TaskData.ReimportObject = ImportAssetParameters.ReimportAsset;
-	TWeakPtr<Interchange::FImportAsyncHelper, ESPMode::ThreadSafe> WeakAsyncHelper = CreateAsyncHelper(TaskData);
-	TSharedPtr<Interchange::FImportAsyncHelper, ESPMode::ThreadSafe> AsyncHelper = WeakAsyncHelper.Pin();
+	TWeakPtr<UE::Interchange::FImportAsyncHelper, ESPMode::ThreadSafe> WeakAsyncHelper = CreateAsyncHelper(TaskData);
+	TSharedPtr<UE::Interchange::FImportAsyncHelper, ESPMode::ThreadSafe> AsyncHelper = WeakAsyncHelper.Pin();
 	check(AsyncHelper.IsValid());
 
 	FText TitleText = NSLOCTEXT("Interchange", "Asynchronous_import_start", "Importing");
@@ -327,7 +327,7 @@ Interchange::FAsyncImportResult UInterchangeManager::ImportAssetAsync(const FStr
 	check(AsyncHelper->Translators.Num() == AsyncHelper->SourceDatas.Num());
 	for (int32 SourceDataIndex = 0; SourceDataIndex < AsyncHelper->SourceDatas.Num(); ++SourceDataIndex)
 	{
-		int32 TranslatorTaskIndex = AsyncHelper->TranslatorTasks.Add(TGraphTask<Interchange::FTaskTranslator>::CreateTask().ConstructAndDispatchWhenReady(SourceDataIndex, WeakAsyncHelper));
+		int32 TranslatorTaskIndex = AsyncHelper->TranslatorTasks.Add(TGraphTask<UE::Interchange::FTaskTranslator>::CreateTask().ConstructAndDispatchWhenReady(SourceDataIndex, WeakAsyncHelper));
 		PipelinePrerequistes.Add(AsyncHelper->TranslatorTasks[TranslatorTaskIndex]);
 	}
 		
@@ -337,7 +337,7 @@ Interchange::FAsyncImportResult UInterchangeManager::ImportAssetAsync(const FStr
 		UInterchangePipelineBase* GraphPipeline = AsyncHelper->Pipelines[GraphPipelineIndex];
 		TWeakObjectPtr<UInterchangePipelineBase> WeakPipelinePtr = GraphPipeline;
 		int32 GraphPipelineTaskIndex = INDEX_NONE;
-		GraphPipelineTaskIndex = AsyncHelper->PipelineTasks.Add(TGraphTask<Interchange::FTaskPipeline>::CreateTask(&PipelinePrerequistes).ConstructAndDispatchWhenReady(WeakPipelinePtr, WeakAsyncHelper));
+		GraphPipelineTaskIndex = AsyncHelper->PipelineTasks.Add(TGraphTask<UE::Interchange::FTaskPipeline>::CreateTask(&PipelinePrerequistes).ConstructAndDispatchWhenReady(WeakPipelinePtr, WeakAsyncHelper));
 		//Ensure we run the pipeline in the same order we create the task, since pipeline modify the node container, its important that its not process in parallel, Adding the one we start to the prerequisites
 		//is the way to go here
 		PipelinePrerequistes.Add(AsyncHelper->PipelineTasks[GraphPipelineTaskIndex]);
@@ -348,17 +348,17 @@ Interchange::FAsyncImportResult UInterchangeManager::ImportAssetAsync(const FStr
 
 	if (GraphParsingPrerequistes.Num() > 0)
 	{
-		AsyncHelper->ParsingTask = TGraphTask<Interchange::FTaskParsing>::CreateTask(&GraphParsingPrerequistes).ConstructAndDispatchWhenReady(this, PackageBasePath, WeakAsyncHelper);
+		AsyncHelper->ParsingTask = TGraphTask<UE::Interchange::FTaskParsing>::CreateTask(&GraphParsingPrerequistes).ConstructAndDispatchWhenReady(this, PackageBasePath, WeakAsyncHelper);
 	}
 	else
 	{
 		//Fallback on the translator pipeline prerequisites (translator must be done if there is no pipeline)
-		AsyncHelper->ParsingTask = TGraphTask<Interchange::FTaskParsing>::CreateTask(&PipelinePrerequistes).ConstructAndDispatchWhenReady(this, PackageBasePath, WeakAsyncHelper);
+		AsyncHelper->ParsingTask = TGraphTask<UE::Interchange::FTaskParsing>::CreateTask(&PipelinePrerequistes).ConstructAndDispatchWhenReady(this, PackageBasePath, WeakAsyncHelper);
 	}
 
 	//The graph parsing task will create the FCreateAssetTask that will run after them, the FAssetImportTask will call the appropriate Post asset import pipeline when the asset is completed
 
-	return Interchange::FAsyncImportResult{ AsyncHelper->RootObject.GetFuture(), AsyncHelper->RootObjectCompletionEvent };
+	return UE::Interchange::FAsyncImportResult{ AsyncHelper->RootObject.GetFuture(), AsyncHelper->RootObjectCompletionEvent };
 }
 
 bool UInterchangeManager::ImportScene(const FString& ImportContext, const UInterchangeSourceData* SourceData, bool bIsReimport, bool bIsAutomated)
@@ -386,9 +386,9 @@ UInterchangeSourceData* UInterchangeManager::CreateSourceData(const FString& InF
 	return SourceDataAsset;
 }
 
-TWeakPtr<Interchange::FImportAsyncHelper, ESPMode::ThreadSafe> UInterchangeManager::CreateAsyncHelper(const Interchange::FImportAsyncHelperData& Data)
+TWeakPtr<UE::Interchange::FImportAsyncHelper, ESPMode::ThreadSafe> UInterchangeManager::CreateAsyncHelper(const UE::Interchange::FImportAsyncHelperData& Data)
 {
-	TSharedPtr<Interchange::FImportAsyncHelper, ESPMode::ThreadSafe> AsyncHelper = MakeShared<Interchange::FImportAsyncHelper, ESPMode::ThreadSafe>();
+	TSharedPtr<UE::Interchange::FImportAsyncHelper, ESPMode::ThreadSafe> AsyncHelper = MakeShared<UE::Interchange::FImportAsyncHelper, ESPMode::ThreadSafe>();
 	//Copy the task data
 	AsyncHelper->TaskData = Data;
 	int32 AsyncHelperIndex = ImportTasks.Add(AsyncHelper);
@@ -403,7 +403,7 @@ TWeakPtr<Interchange::FImportAsyncHelper, ESPMode::ThreadSafe> UInterchangeManag
 	return ImportTasks[AsyncHelperIndex];
 }
 
-void UInterchangeManager::ReleaseAsyncHelper(TWeakPtr<Interchange::FImportAsyncHelper, ESPMode::ThreadSafe> AsyncHelper)
+void UInterchangeManager::ReleaseAsyncHelper(TWeakPtr<UE::Interchange::FImportAsyncHelper, ESPMode::ThreadSafe> AsyncHelper)
 {
 	check(AsyncHelper.IsValid());
 	ImportTasks.RemoveSingle(AsyncHelper.Pin());
diff --git a/Engine/Source/Runtime/Interchange/Core/Private/Nodes/InterchangeBaseNode.cpp b/Engine/Source/Runtime/Interchange/Core/Private/Nodes/InterchangeBaseNode.cpp
--- a/Engine/Source/Runtime/Interchange/Core/Private/Nodes/InterchangeBaseNode.cpp
+++ b/Engine/Source/Runtime/Interchange/Core/Private/Nodes/InterchangeBaseNode.cpp
@@ -12,16 +12,16 @@
 
 void UInterchangeBaseNode::InitializeNode(const FName& UniqueID, const FName& DisplayLabel)
 {
-	Interchange::EAttributeStorageResult Result = Attributes.RegisterAttribute(Interchange::FBaseNodeStaticData::UniqueIDKey(), UniqueID, Interchange::EAttributeProperty::NoHash);
+	UE::Interchange::EAttributeStorageResult Result = Attributes.RegisterAttribute(UE::Interchange::FBaseNodeStaticData::UniqueIDKey(), UniqueID, UE::Interchange::EAttributeProperty::NoHash);
 	if (!IsAttributeStorageResultSuccess(Result))
 	{
-		LogAttributeStorageErrors(Result, TEXT("RegisterAttribute"), Interchange::FBaseNodeStaticData::UniqueIDKey());
+		LogAttributeStorageErrors(Result, TEXT("RegisterAttribute"), UE::Interchange::FBaseNodeStaticData::UniqueIDKey());
 	}
 
-	Result = Attributes.RegisterAttribute(Interchange::FBaseNodeStaticData::DisplayLabelKey(), DisplayLabel, Interchange::EAttributeProperty::NoHash);
+	Result = Attributes.RegisterAttribute(UE::Interchange::FBaseNodeStaticData::DisplayLabelKey(), DisplayLabel, UE::Interchange::EAttributeProperty::NoHash);
 	if (!IsAttributeStorageResultSuccess(Result))
 	{
-		LogAttributeStorageErrors(Result, TEXT("RegisterAttribute"), Interchange::FBaseNodeStaticData::DisplayLabelKey());
+		LogAttributeStorageErrors(Result, TEXT("RegisterAttribute"), UE::Interchange::FBaseNodeStaticData::DisplayLabelKey());
 	}
 	bIsInitialized = true;
 }
@@ -32,7 +32,7 @@ FString UInterchangeBaseNode::GetTypeName() const
 	return TypeName;
 }
 
-bool UInterchangeBaseNode::HasAttribute(const Interchange::FAttributeKey& NodeAttributeKey) const
+bool UInterchangeBaseNode::HasAttribute(const UE::Interchange::FAttributeKey& NodeAttributeKey) const
 {
 	return Attributes.ContainAttribute(NodeAttributeKey);
 }
@@ -42,28 +42,28 @@ FName UInterchangeBaseNode::GetUniqueID() const
 {
 	ensure(bIsInitialized);
 	FName UniqueID = NAME_None;
-	Attributes.GetAttributeHandle<FName>(Interchange::FBaseNodeStaticData::UniqueIDKey()).Get(UniqueID);
+	Attributes.GetAttributeHandle<FName>(UE::Interchange::FBaseNodeStaticData::UniqueIDKey()).Get(UniqueID);
 	return UniqueID;
 }
 
 FName UInterchangeBaseNode::GetDisplayLabel() const
 {
 	ensure(bIsInitialized);
-	checkSlow(Attributes.ContainAttribute(Interchange::FBaseNodeStaticData::DisplayLabelKey()));
+	checkSlow(Attributes.ContainAttribute(UE::Interchange::FBaseNodeStaticData::DisplayLabelKey()));
 	FName DisplayLabel = NAME_None;
-	Attributes.GetAttributeHandle<FName>(Interchange::FBaseNodeStaticData::DisplayLabelKey()).Get(DisplayLabel);
+	Attributes.GetAttributeHandle<FName>(UE::Interchange::FBaseNodeStaticData::DisplayLabelKey()).Get(DisplayLabel);
 	return DisplayLabel;
 }
 
 FName UInterchangeBaseNode::GetParentUID() const
 {
-	if (!Attributes.ContainAttribute(Interchange::FBaseNodeStaticData::ParentIDKey()))
+	if (!Attributes.ContainAttribute(UE::Interchange::FBaseNodeStaticData::ParentIDKey()))
 	{
 		return InvalidNodeUID();
 	}
 
 	FName ParentUniqueID = NAME_None;
-	Interchange::FAttributeStorage::TAttributeHandle<FName> Handle = Attributes.GetAttributeHandle<FName>(Interchange::FBaseNodeStaticData::ParentIDKey());
+	UE::Interchange::FAttributeStorage::TAttributeHandle<FName> Handle = Attributes.GetAttributeHandle<FName>(UE::Interchange::FBaseNodeStaticData::ParentIDKey());
 	if(Handle.IsValid())
 	{
 		Handle.Get(ParentUniqueID);
@@ -74,10 +74,10 @@ FName UInterchangeBaseNode::GetParentUID() const
 
 bool UInterchangeBaseNode::SetParentUID(FName ParentUID)
 {
-	Interchange::EAttributeStorageResult Result = Attributes.RegisterAttribute(Interchange::FBaseNodeStaticData::ParentIDKey(), ParentUID);
+	UE::Interchange::EAttributeStorageResult Result = Attributes.RegisterAttribute(UE::Interchange::FBaseNodeStaticData::ParentIDKey(), ParentUID);
 	if(IsAttributeStorageResultSuccess(Result))
 	{
-		Interchange::FAttributeStorage::TAttributeHandle<FName> Handle = Attributes.GetAttributeHandle<FName>(Interchange::FBaseNodeStaticData::ParentIDKey());
+		UE::Interchange::FAttributeStorage::TAttributeHandle<FName> Handle = Attributes.GetAttributeHandle<FName>(UE::Interchange::FBaseNodeStaticData::ParentIDKey());
 		return Handle.IsValid();
 	}
 	return false;
@@ -88,11 +88,11 @@ void UInterchangeBaseNode::GetDependecies(TArray<FName>& OutDependencies ) const
 	OutDependencies.Reset();
 
 	int32 DepenedenciesCount = 0;
-	if (!Attributes.ContainAttribute(Interchange::FBaseNodeStaticData::DependencyCountKey()))
+	if (!Attributes.ContainAttribute(UE::Interchange::FBaseNodeStaticData::DependencyCountKey()))
 	{
 		return;
 	}
-	Interchange::FAttributeStorage::TAttributeHandle<int32> Handle = Attributes.GetAttributeHandle<int32>(Interchange::FBaseNodeStaticData::DependencyCountKey());
+	UE::Interchange::FAttributeStorage::TAttributeHandle<int32> Handle = Attributes.GetAttributeHandle<int32>(UE::Interchange::FBaseNodeStaticData::DependencyCountKey());
 	if (!Handle.IsValid())
 	{
 		return;
@@ -100,9 +100,9 @@ void UInterchangeBaseNode::GetDependecies(TArray<FName>& OutDependencies ) const
 	Handle.Get(DepenedenciesCount);
 	for (int32 DepIndex = 0; DepIndex < DepenedenciesCount; ++DepIndex)
 	{
-		FString DepIndexKeyString = Interchange::FBaseNodeStaticData::DependencyBaseKey().Key.ToString() + FString::FromInt(DepIndex);
-		Interchange::FAttributeKey DepIndexKey(*DepIndexKeyString);
-		Interchange::FAttributeStorage::TAttributeHandle<FName> HandleDep = Attributes.GetAttributeHandle<FName>(DepIndexKey);
+		FString DepIndexKeyString = UE::Interchange::FBaseNodeStaticData::DependencyBaseKey().Key.ToString() + FString::FromInt(DepIndex);
+		UE::Interchange::FAttributeKey DepIndexKey(*DepIndexKeyString);
+		UE::Interchange::FAttributeStorage::TAttributeHandle<FName> HandleDep = Attributes.GetAttributeHandle<FName>(DepIndexKey);
 		if (!HandleDep.IsValid())
 		{
 			continue;
@@ -114,31 +114,31 @@ void UInterchangeBaseNode::GetDependecies(TArray<FName>& OutDependencies ) const
 
 bool UInterchangeBaseNode::SetDependencyUID(FName DependencyUID)
 {
-	if (!Attributes.ContainAttribute(Interchange::FBaseNodeStaticData::DependencyCountKey()))
+	if (!Attributes.ContainAttribute(UE::Interchange::FBaseNodeStaticData::DependencyCountKey()))
 	{
 		const int32 DependencyCount = 0;
-		Interchange::EAttributeStorageResult Result = Attributes.RegisterAttribute<int32>(Interchange::FBaseNodeStaticData::DependencyCountKey(), DependencyCount);
+		UE::Interchange::EAttributeStorageResult Result = Attributes.RegisterAttribute<int32>(UE::Interchange::FBaseNodeStaticData::DependencyCountKey(), DependencyCount);
 	}
-	Interchange::FAttributeStorage::TAttributeHandle<int32> Handle = Attributes.GetAttributeHandle<int32>(Interchange::FBaseNodeStaticData::DependencyCountKey());
+	UE::Interchange::FAttributeStorage::TAttributeHandle<int32> Handle = Attributes.GetAttributeHandle<int32>(UE::Interchange::FBaseNodeStaticData::DependencyCountKey());
 	if(!ensure(Handle.IsValid()))
 	{
 		return false;
 	}
 	int32 DepIndex = 0;
 	Handle.Get(DepIndex);
-	FString DepIndexKeyString = Interchange::FBaseNodeStaticData::DependencyBaseKey().Key.ToString() + FString::FromInt(DepIndex);
-	Interchange::FAttributeKey DepIndexKey(*DepIndexKeyString);
+	FString DepIndexKeyString = UE::Interchange::FBaseNodeStaticData::DependencyBaseKey().Key.ToString() + FString::FromInt(DepIndex);
+	UE::Interchange::FAttributeKey DepIndexKey(*DepIndexKeyString);
 	//Increment the counter
 	DepIndex++;
 	Handle.Set(DepIndex);
 
-	Interchange::EAttributeStorageResult DepResult = Attributes.RegisterAttribute<FName>(DepIndexKey, DependencyUID);
+	UE::Interchange::EAttributeStorageResult DepResult = Attributes.RegisterAttribute<FName>(DepIndexKey, DependencyUID);
 	return true;
 }
 
 bool UInterchangeBaseNode::IsEnabled() const
 {
-	Interchange::FAttributeStorage::TAttributeHandle<bool> Handle = Attributes.GetAttributeHandle<bool>(Interchange::FBaseNodeStaticData::IsEnabledKey());
+	UE::Interchange::FAttributeStorage::TAttributeHandle<bool> Handle = Attributes.GetAttributeHandle<bool>(UE::Interchange::FBaseNodeStaticData::IsEnabledKey());
 	if (Handle.IsValid())
 	{
 		bool bIsEnabled = false;
@@ -150,10 +150,10 @@ bool UInterchangeBaseNode::IsEnabled() const
 
 bool UInterchangeBaseNode::SetEnabled(const bool bIsEnabled)
 {
-	Interchange::EAttributeStorageResult Result = Attributes.RegisterAttribute(Interchange::FBaseNodeStaticData::IsEnabledKey(), bIsEnabled);
+	UE::Interchange::EAttributeStorageResult Result = Attributes.RegisterAttribute(UE::Interchange::FBaseNodeStaticData::IsEnabledKey(), bIsEnabled);
 	if (IsAttributeStorageResultSuccess(Result))
 	{
-		Interchange::FAttributeStorage::TAttributeHandle<bool> Handle = Attributes.GetAttributeHandle<bool>(Interchange::FBaseNodeStaticData::IsEnabledKey());
+		UE::Interchange::FAttributeStorage::TAttributeHandle<bool> Handle = Attributes.GetAttributeHandle<bool>(UE::Interchange::FBaseNodeStaticData::IsEnabledKey());
 		return Handle.IsValid();
 	}
 	return false;
@@ -177,11 +177,11 @@ FName UInterchangeBaseNode::InvalidNodeUID()
 void UInterchangeBaseNode::ApplyAllCustomAttributeToAsset(UObject* Object) const
 {
 	UClass* ObjectClass = Object->GetClass();
-	for (const TPair<UClass*, TArray<Interchange::FApplyAttributeToAsset>>& ClassDelegatePair : ApplyCustomAttributeDelegates)
+	for (const TPair<UClass*, TArray<UE::Interchange::FApplyAttributeToAsset>>& ClassDelegatePair : ApplyCustomAttributeDelegates)
 	{
 		if(ObjectClass->IsChildOf(ClassDelegatePair.Key))
 		{
-			for (const Interchange::FApplyAttributeToAsset& Delegate : ClassDelegatePair.Value)
+			for (const UE::Interchange::FApplyAttributeToAsset& Delegate : ClassDelegatePair.Value)
 			{
 				if(Delegate.IsBound())
 				{
@@ -198,7 +198,7 @@ void UInterchangeBaseNode::Serialize(FArchive& Ar)
 	if (Ar.IsLoading())
 	{
 		//The node is consider Initialize if the UniqueID and the Display label are set properly
-		bIsInitialized = (Attributes.GetAttributeHandle<FName>(Interchange::FBaseNodeStaticData::UniqueIDKey()).IsValid() &&
-							Attributes.GetAttributeHandle<FName>(Interchange::FBaseNodeStaticData::DisplayLabelKey()).IsValid());
+		bIsInitialized = (Attributes.GetAttributeHandle<FName>(UE::Interchange::FBaseNodeStaticData::UniqueIDKey()).IsValid() &&
+							Attributes.GetAttributeHandle<FName>(UE::Interchange::FBaseNodeStaticData::DisplayLabelKey()).IsValid());
 	}
 }
diff --git a/Engine/Source/Runtime/Interchange/Core/Private/Tasks/InterchangeTaskCompletion.cpp b/Engine/Source/Runtime/Interchange/Core/Private/Tasks/InterchangeTaskCompletion.cpp
--- a/Engine/Source/Runtime/Interchange/Core/Private/Tasks/InterchangeTaskCompletion.cpp
+++ b/Engine/Source/Runtime/Interchange/Core/Private/Tasks/InterchangeTaskCompletion.cpp
@@ -15,16 +15,16 @@
 #include "AssetRegistryModule.h"
 #endif //WITH_ENGINE
 
-void Interchange::FTaskCompletion::DoTask(ENamedThreads::Type CurrentThread, const FGraphEventRef& MyCompletionGraphEvent)
+void UE::Interchange::FTaskCompletion::DoTask(ENamedThreads::Type CurrentThread, const FGraphEventRef& MyCompletionGraphEvent)
 {
-	TSharedPtr<Interchange::FImportAsyncHelper, ESPMode::ThreadSafe> AsyncHelper = WeakAsyncHelper.Pin();
+	TSharedPtr<FImportAsyncHelper, ESPMode::ThreadSafe> AsyncHelper = WeakAsyncHelper.Pin();
 	check(AsyncHelper.IsValid());
 	bool bIsFutureRootObjectSet = false;
-	for(TPair<int32, TArray<Interchange::FImportAsyncHelper::FImportedAssetInfo>>& AssetInfosPerSourceIndexPair : AsyncHelper->ImportedAssetsPerSourceIndex)
+	for(TPair<int32, TArray<FImportAsyncHelper::FImportedAssetInfo>>& AssetInfosPerSourceIndexPair : AsyncHelper->ImportedAssetsPerSourceIndex)
 	{
 		const int32 SourceIndex = AssetInfosPerSourceIndexPair.Key;
 		const bool bCallPostImportGameThreadCallback = ensure(AsyncHelper->SourceDatas.IsValidIndex(SourceIndex));
-		for (const Interchange::FImportAsyncHelper::FImportedAssetInfo& AssetInfo : AssetInfosPerSourceIndexPair.Value)
+		for (const FImportAsyncHelper::FImportedAssetInfo& AssetInfo : AssetInfosPerSourceIndexPair.Value)
 		{
 			UObject* Asset = AssetInfo.ImportAsset;
 			//In case Some factory code cannot run outside of the main thread we offer this callback to finish the work before calling post edit change (building the asset)
diff --git a/Engine/Source/Runtime/Interchange/Core/Private/Tasks/InterchangeTaskCreateAsset.cpp b/Engine/Source/Runtime/Interchange/Core/Private/Tasks/InterchangeTaskCreateAsset.cpp
--- a/Engine/Source/Runtime/Interchange/Core/Private/Tasks/InterchangeTaskCreateAsset.cpp
+++ b/Engine/Source/Runtime/Interchange/Core/Private/Tasks/InterchangeTaskCreateAsset.cpp
@@ -18,34 +18,40 @@
 #include "UObject/WeakObjectPtrTemplates.h"
 #include "Nodes/InterchangeBaseNode.h"
 
-namespace Interchange_InternalImplementation
+namespace UE
 {
-	void InternalGetPackageName(const Interchange::FImportAsyncHelper& AsyncHelper, const int32 SourceIndex, const FString& PackageBasePath, const UInterchangeBaseNode* Node, FString& OutPackageName, FString& OutAssetName)
+	namespace Interchange
 	{
-		const UInterchangeSourceData* SourceData = AsyncHelper.SourceDatas[SourceIndex];
-		check(SourceData);
-		FString NodeDisplayName = Node->GetDisplayLabel().ToString();
-		const FString BaseFileName = FPaths::GetBaseFilename(SourceData->GetFilename());
-
-		//Set the asset name and the package name
-		if (NodeDisplayName.Equals(BaseFileName) || BaseFileName.IsEmpty())
-		{
-			OutAssetName = NodeDisplayName;
-		}
-		else
+		namespace Private
 		{
-			OutAssetName = BaseFileName.IsEmpty() ? NodeDisplayName : BaseFileName + TEXT("_") + NodeDisplayName;
-		}
-		OutPackageName = FPaths::Combine(*PackageBasePath, *OutAssetName);
-		
-		//Sanitize only the package name
-		Interchange::SanitizeInvalidChar(OutPackageName);
-	}
-}
-
-void Interchange::FTaskCreatePackage::DoTask(ENamedThreads::Type CurrentThread, const FGraphEventRef& MyCompletionGraphEvent)
+			void InternalGetPackageName(const UE::Interchange::FImportAsyncHelper& AsyncHelper, const int32 SourceIndex, const FString& PackageBasePath, const UInterchangeBaseNode* Node, FString& OutPackageName, FString& OutAssetName)
+			{
+				const UInterchangeSourceData* SourceData = AsyncHelper.SourceDatas[SourceIndex];
+				check(SourceData);
+				FString NodeDisplayName = Node->GetDisplayLabel().ToString();
+				const FString BaseFileName = FPaths::GetBaseFilename(SourceData->GetFilename());
+
+				//Set the asset name and the package name
+				if (NodeDisplayName.Equals(BaseFileName) || BaseFileName.IsEmpty())
+				{
+					OutAssetName = NodeDisplayName;
+				}
+				else
+				{
+					OutAssetName = BaseFileName.IsEmpty() ? NodeDisplayName : BaseFileName + TEXT("_") + NodeDisplayName;
+				}
+				OutPackageName = FPaths::Combine(*PackageBasePath, *OutAssetName);
+
+				//Sanitize only the package name
+				UE::Interchange::SanitizeInvalidChar(OutPackageName);
+			}
+		}//ns Private
+	}//ns Interchange
+}//ns UE
+
+void UE::Interchange::FTaskCreatePackage::DoTask(ENamedThreads::Type CurrentThread, const FGraphEventRef& MyCompletionGraphEvent)
 {
-	TSharedPtr<Interchange::FImportAsyncHelper, ESPMode::ThreadSafe> AsyncHelper = WeakAsyncHelper.Pin();
+	TSharedPtr<UE::Interchange::FImportAsyncHelper, ESPMode::ThreadSafe> AsyncHelper = WeakAsyncHelper.Pin();
 	check(AsyncHelper.IsValid());
 
 	//The create package thread must always execute on the game thread
@@ -62,9 +68,9 @@ void Interchange::FTaskCreatePackage::DoTask(ENamedThreads::Type CurrentThread,
 	}
 	else
 	{
-		Interchange_InternalImplementation::InternalGetPackageName(*AsyncHelper, SourceIndex, PackageBasePath, Node, PackageName, AssetName);
+		Private::InternalGetPackageName(*AsyncHelper, SourceIndex, PackageBasePath, Node, PackageName, AssetName);
 		// We can not create assets that share the name of a map file in the same location
-		if (Interchange::FPackageUtils::IsMapPackageAsset(PackageName))
+		if (UE::Interchange::FPackageUtils::IsMapPackageAsset(PackageName))
 		{
 			const FText Message = FText::Format(NSLOCTEXT("Interchange", "AssetNameInUseByMap", "You can not create an asset named '{0}' because there is already a map file with this name in this folder."), FText::FromString(AssetName));
 			UE_LOG(LogInterchangeCore, Warning, TEXT("%s"), *Message.ToString());
@@ -103,14 +109,14 @@ void Interchange::FTaskCreatePackage::DoTask(ENamedThreads::Type CurrentThread,
 	AsyncHelper->CreatedPackages.Add(PackageName, Pkg);
 }
 
-void Interchange::FTaskCreateAsset::DoTask(ENamedThreads::Type CurrentThread, const FGraphEventRef& MyCompletionGraphEvent)
+void UE::Interchange::FTaskCreateAsset::DoTask(ENamedThreads::Type CurrentThread, const FGraphEventRef& MyCompletionGraphEvent)
 {
-	TSharedPtr<Interchange::FImportAsyncHelper, ESPMode::ThreadSafe> AsyncHelper = WeakAsyncHelper.Pin();
+	TSharedPtr<UE::Interchange::FImportAsyncHelper, ESPMode::ThreadSafe> AsyncHelper = WeakAsyncHelper.Pin();
 	check(WeakAsyncHelper.IsValid());
 
 	FString PackageName;
 	FString AssetName;
-	Interchange_InternalImplementation::InternalGetPackageName(*AsyncHelper, SourceIndex, PackageBasePath, Node, PackageName, AssetName);
+	Private::InternalGetPackageName(*AsyncHelper, SourceIndex, PackageBasePath, Node, PackageName, AssetName);
 	if (AsyncHelper->TaskData.ReimportObject)
 	{
 		UPackage* Pkg = AsyncHelper->TaskData.ReimportObject->GetPackage();
@@ -157,8 +163,8 @@ void Interchange::FTaskCreateAsset::DoTask(ENamedThreads::Type CurrentThread, co
 	if (NodeAsset)
 	{
 		FScopeLock Lock(&AsyncHelper->ImportedAssetsPerSourceIndexLock);
-		TArray<Interchange::FImportAsyncHelper::FImportedAssetInfo>& ImportedInfos = AsyncHelper->ImportedAssetsPerSourceIndex.FindOrAdd(SourceIndex);
-		Interchange::FImportAsyncHelper::FImportedAssetInfo& AssetInfo = ImportedInfos.AddZeroed_GetRef();
+		TArray<UE::Interchange::FImportAsyncHelper::FImportedAssetInfo>& ImportedInfos = AsyncHelper->ImportedAssetsPerSourceIndex.FindOrAdd(SourceIndex);
+		UE::Interchange::FImportAsyncHelper::FImportedAssetInfo& AssetInfo = ImportedInfos.AddZeroed_GetRef();
 		AssetInfo.ImportAsset = NodeAsset;
 		Node->ReferenceObject = FSoftObjectPath(NodeAsset);
 		AssetInfo.Factory = Factory;
diff --git a/Engine/Source/Runtime/Interchange/Core/Private/Tasks/InterchangeTaskCreateAsset.h b/Engine/Source/Runtime/Interchange/Core/Private/Tasks/InterchangeTaskCreateAsset.h
--- a/Engine/Source/Runtime/Interchange/Core/Private/Tasks/InterchangeTaskCreateAsset.h
+++ b/Engine/Source/Runtime/Interchange/Core/Private/Tasks/InterchangeTaskCreateAsset.h
@@ -10,94 +10,97 @@
 #include "UObject/WeakObjectPtrTemplates.h"
 #include "Nodes/InterchangeBaseNode.h"
 
-namespace Interchange
+namespace UE
 {
-
-/**
- * This task create package, Cook::PackageTracker::NotifyUObjectCreated is not thread safe, so we need to create the packages on the main thread
- */
-class FTaskCreatePackage
-{
-private:
-	FString PackageBasePath;
-	int32 SourceIndex;
-	TWeakPtr<Interchange::FImportAsyncHelper, ESPMode::ThreadSafe> WeakAsyncHelper;
-	const UInterchangeBaseNode* Node;
-	UInterchangeFactoryBase* Factory;
-
-public:
-	FTaskCreatePackage(const FString& InPackageBasePath, const int32 InSourceIndex, TWeakPtr<Interchange::FImportAsyncHelper, ESPMode::ThreadSafe> InAsyncHelper, const UInterchangeBaseNode* InNode, UInterchangeFactoryBase* InFactory)
-		: PackageBasePath(InPackageBasePath)
-		, SourceIndex(InSourceIndex)
-		, WeakAsyncHelper(InAsyncHelper)
-		, Node(InNode)
-		, Factory(InFactory)
+	namespace Interchange
 	{
-		check(Node);
-		check(Factory);
-	}
 
-	ENamedThreads::Type GetDesiredThread()
-	{
-		if (WeakAsyncHelper.IsValid() && WeakAsyncHelper.Pin()->TaskData.ReimportObject)
+		/**
+		 * This task create package, Cook::PackageTracker::NotifyUObjectCreated is not thread safe, so we need to create the packages on the main thread
+		 */
+		class FTaskCreatePackage
 		{
-			//When doing a reimport the package already exist, so we can get it outside of the main thread
-			return ENamedThreads::AnyBackgroundThreadNormalTask;
-		}
-		return ENamedThreads::GameThread;
-	}
-
-	static ESubsequentsMode::Type GetSubsequentsMode()
-	{
-		return ESubsequentsMode::TrackSubsequents;
-	}
-
-	TStatId GetStatId() const
-	{
-		RETURN_QUICK_DECLARE_CYCLE_STAT(FTaskCreatePackage, STATGROUP_TaskGraphTasks);
-	}
-
-	void DoTask(ENamedThreads::Type CurrentThread, const FGraphEventRef& MyCompletionGraphEvent);
-};
-
-class FTaskCreateAsset
-{
-private:
-	FString PackageBasePath;
-	int32 SourceIndex;
-	TWeakPtr<Interchange::FImportAsyncHelper, ESPMode::ThreadSafe> WeakAsyncHelper;
-	const UInterchangeBaseNode* Node;
-	UInterchangeFactoryBase* Factory;
-
-public:
-	FTaskCreateAsset(const FString& InPackageBasePath, const int32 InSourceIndex, TWeakPtr<Interchange::FImportAsyncHelper, ESPMode::ThreadSafe> InAsyncHelper, const UInterchangeBaseNode* InNode, UInterchangeFactoryBase* InFactory)
-		: PackageBasePath(InPackageBasePath)
-		, SourceIndex(InSourceIndex)
-		, WeakAsyncHelper(InAsyncHelper)
-		, Node(InNode)
-		, Factory(InFactory)
-	{
-		check(Node);
-		check(Factory);
-	}
-
-	ENamedThreads::Type GetDesiredThread()
-	{
-		return Factory->CanExecuteOnAnyThread() ? ENamedThreads::AnyBackgroundThreadNormalTask : ENamedThreads::GameThread;
-	}
-
-	static ESubsequentsMode::Type GetSubsequentsMode()
-	{
-		return ESubsequentsMode::TrackSubsequents;
-	}
-
-	TStatId GetStatId() const
-	{
-		RETURN_QUICK_DECLARE_CYCLE_STAT(FTaskCreateAsset, STATGROUP_TaskGraphTasks);
-	}
-
-	void DoTask(ENamedThreads::Type CurrentThread, const FGraphEventRef& MyCompletionGraphEvent);
-};
-
-
-} // End namespace Interchange
+		private:
+			FString PackageBasePath;
+			int32 SourceIndex;
+			TWeakPtr<FImportAsyncHelper, ESPMode::ThreadSafe> WeakAsyncHelper;
+			const UInterchangeBaseNode* Node;
+			UInterchangeFactoryBase* Factory;
+
+		public:
+			FTaskCreatePackage(const FString& InPackageBasePath, const int32 InSourceIndex, TWeakPtr<FImportAsyncHelper, ESPMode::ThreadSafe> InAsyncHelper, const UInterchangeBaseNode* InNode, UInterchangeFactoryBase* InFactory)
+				: PackageBasePath(InPackageBasePath)
+				, SourceIndex(InSourceIndex)
+				, WeakAsyncHelper(InAsyncHelper)
+				, Node(InNode)
+				, Factory(InFactory)
+			{
+				check(Node);
+				check(Factory);
+			}
+
+			ENamedThreads::Type GetDesiredThread()
+			{
+				if (WeakAsyncHelper.IsValid() && WeakAsyncHelper.Pin()->TaskData.ReimportObject)
+				{
+					//When doing a reimport the package already exist, so we can get it outside of the main thread
+					return ENamedThreads::AnyBackgroundThreadNormalTask;
+				}
+				return ENamedThreads::GameThread;
+			}
+
+			static ESubsequentsMode::Type GetSubsequentsMode()
+			{
+				return ESubsequentsMode::TrackSubsequents;
+			}
+
+			TStatId GetStatId() const
+			{
+				RETURN_QUICK_DECLARE_CYCLE_STAT(FTaskCreatePackage, STATGROUP_TaskGraphTasks);
+			}
+
+			void DoTask(ENamedThreads::Type CurrentThread, const FGraphEventRef& MyCompletionGraphEvent);
+		};
+
+		class FTaskCreateAsset
+		{
+		private:
+			FString PackageBasePath;
+			int32 SourceIndex;
+			TWeakPtr<FImportAsyncHelper, ESPMode::ThreadSafe> WeakAsyncHelper;
+			const UInterchangeBaseNode* Node;
+			UInterchangeFactoryBase* Factory;
+
+		public:
+			FTaskCreateAsset(const FString& InPackageBasePath, const int32 InSourceIndex, TWeakPtr<FImportAsyncHelper, ESPMode::ThreadSafe> InAsyncHelper, const UInterchangeBaseNode* InNode, UInterchangeFactoryBase* InFactory)
+				: PackageBasePath(InPackageBasePath)
+				, SourceIndex(InSourceIndex)
+				, WeakAsyncHelper(InAsyncHelper)
+				, Node(InNode)
+				, Factory(InFactory)
+			{
+				check(Node);
+				check(Factory);
+			}
+
+			ENamedThreads::Type GetDesiredThread()
+			{
+				return Factory->CanExecuteOnAnyThread() ? ENamedThreads::AnyBackgroundThreadNormalTask : ENamedThreads::GameThread;
+			}
+
+			static ESubsequentsMode::Type GetSubsequentsMode()
+			{
+				return ESubsequentsMode::TrackSubsequents;
+			}
+
+			TStatId GetStatId() const
+			{
+				RETURN_QUICK_DECLARE_CYCLE_STAT(FTaskCreateAsset, STATGROUP_TaskGraphTasks);
+			}
+
+			void DoTask(ENamedThreads::Type CurrentThread, const FGraphEventRef& MyCompletionGraphEvent);
+		};
+
+
+	} //ns Interchange
+}//ns UE
diff --git a/Engine/Source/Runtime/Interchange/Core/Private/Tasks/InterchangeTaskParsing.cpp b/Engine/Source/Runtime/Interchange/Core/Private/Tasks/InterchangeTaskParsing.cpp
--- a/Engine/Source/Runtime/Interchange/Core/Private/Tasks/InterchangeTaskParsing.cpp
+++ b/Engine/Source/Runtime/Interchange/Core/Private/Tasks/InterchangeTaskParsing.cpp
@@ -20,9 +20,9 @@
 #include "Nodes/InterchangeBaseNode.h"
 #include "Nodes/InterchangeBaseNodeContainer.h"
 
-void Interchange::FTaskParsing::DoTask(ENamedThreads::Type CurrentThread, const FGraphEventRef& MyCompletionGraphEvent)
+void UE::Interchange::FTaskParsing::DoTask(ENamedThreads::Type CurrentThread, const FGraphEventRef& MyCompletionGraphEvent)
 {
-	TSharedPtr<Interchange::FImportAsyncHelper, ESPMode::ThreadSafe> AsyncHelper = WeakAsyncHelper.Pin();
+	TSharedPtr<FImportAsyncHelper, ESPMode::ThreadSafe> AsyncHelper = WeakAsyncHelper.Pin();
 	check(AsyncHelper.IsValid());
 
 	//Parse each graph and prepare import task data, we will then be able to create all the task with the correct dependencies
diff --git a/Engine/Source/Runtime/Interchange/Core/Private/Tasks/InterchangeTaskParsing.h b/Engine/Source/Runtime/Interchange/Core/Private/Tasks/InterchangeTaskParsing.h
--- a/Engine/Source/Runtime/Interchange/Core/Private/Tasks/InterchangeTaskParsing.h
+++ b/Engine/Source/Runtime/Interchange/Core/Private/Tasks/InterchangeTaskParsing.h
@@ -8,40 +8,43 @@
 #include "Stats/Stats.h"
 #include "UObject/WeakObjectPtrTemplates.h"
 
-namespace Interchange
+namespace UE
 {
-
-class FTaskParsing
-{
-private:
-	UInterchangeManager* InterchangeManager;
-	TWeakPtr<Interchange::FImportAsyncHelper, ESPMode::ThreadSafe> WeakAsyncHelper;
-	FString PackageBasePath;
-public:
-	FTaskParsing(UInterchangeManager* InInterchangeManager, const FString& InPackageBasePath, TWeakPtr<Interchange::FImportAsyncHelper, ESPMode::ThreadSafe> InAsyncHelper)
-		: InterchangeManager(InInterchangeManager)
-		, WeakAsyncHelper(InAsyncHelper)
-		, PackageBasePath(InPackageBasePath)
-	{
-		check(InterchangeManager);
-	}
-
-	static FORCEINLINE ENamedThreads::Type GetDesiredThread()
+	namespace Interchange
 	{
-		return ENamedThreads::AnyBackgroundThreadNormalTask;
-	}
-	static FORCEINLINE ESubsequentsMode::Type GetSubsequentsMode()
-	{
-		
-		return ESubsequentsMode::TrackSubsequents;
-	}
-
-	FORCEINLINE TStatId GetStatId() const
-	{
-		RETURN_QUICK_DECLARE_CYCLE_STAT(FTaskParsing, STATGROUP_TaskGraphTasks);
-	}
-
-	void DoTask(ENamedThreads::Type CurrentThread, const FGraphEventRef& MyCompletionGraphEvent);
-};
 
-} // End namespace Interchange
+		class FTaskParsing
+		{
+		private:
+			UInterchangeManager* InterchangeManager;
+			TWeakPtr<FImportAsyncHelper, ESPMode::ThreadSafe> WeakAsyncHelper;
+			FString PackageBasePath;
+		public:
+			FTaskParsing(UInterchangeManager* InInterchangeManager, const FString& InPackageBasePath, TWeakPtr<FImportAsyncHelper, ESPMode::ThreadSafe> InAsyncHelper)
+				: InterchangeManager(InInterchangeManager)
+				, WeakAsyncHelper(InAsyncHelper)
+				, PackageBasePath(InPackageBasePath)
+			{
+				check(InterchangeManager);
+			}
+
+			static FORCEINLINE ENamedThreads::Type GetDesiredThread()
+			{
+				return ENamedThreads::AnyBackgroundThreadNormalTask;
+			}
+			static FORCEINLINE ESubsequentsMode::Type GetSubsequentsMode()
+			{
+
+				return ESubsequentsMode::TrackSubsequents;
+			}
+
+			FORCEINLINE TStatId GetStatId() const
+			{
+				RETURN_QUICK_DECLARE_CYCLE_STAT(FTaskParsing, STATGROUP_TaskGraphTasks);
+			}
+
+			void DoTask(ENamedThreads::Type CurrentThread, const FGraphEventRef& MyCompletionGraphEvent);
+		};
+
+	} //ns Interchange
+}//ns UE
diff --git a/Engine/Source/Runtime/Interchange/Core/Private/Tasks/InterchangeTaskPipeline.cpp b/Engine/Source/Runtime/Interchange/Core/Private/Tasks/InterchangeTaskPipeline.cpp
--- a/Engine/Source/Runtime/Interchange/Core/Private/Tasks/InterchangeTaskPipeline.cpp
+++ b/Engine/Source/Runtime/Interchange/Core/Private/Tasks/InterchangeTaskPipeline.cpp
@@ -10,9 +10,9 @@
 #include "UObject/WeakObjectPtrTemplates.h"
 #include "Nodes/InterchangeBaseNodeContainer.h"
 
-void Interchange::FTaskPipeline::DoTask(ENamedThreads::Type CurrentThread, const FGraphEventRef& MyCompletionGraphEvent)
+void UE::Interchange::FTaskPipeline::DoTask(ENamedThreads::Type CurrentThread, const FGraphEventRef& MyCompletionGraphEvent)
 {
-	TSharedPtr<Interchange::FImportAsyncHelper, ESPMode::ThreadSafe> AsyncHelper = WeakAsyncHelper.Pin();
+	TSharedPtr<FImportAsyncHelper, ESPMode::ThreadSafe> AsyncHelper = WeakAsyncHelper.Pin();
 	check(AsyncHelper.IsValid());
 
 	if (UInterchangePipelineBase* Pipeline = PipelineBase.Get())
diff --git a/Engine/Source/Runtime/Interchange/Core/Private/Tasks/InterchangeTaskPipeline.h b/Engine/Source/Runtime/Interchange/Core/Private/Tasks/InterchangeTaskPipeline.h
--- a/Engine/Source/Runtime/Interchange/Core/Private/Tasks/InterchangeTaskPipeline.h
+++ b/Engine/Source/Runtime/Interchange/Core/Private/Tasks/InterchangeTaskPipeline.h
@@ -9,38 +9,41 @@
 #include "Stats/Stats.h"
 #include "UObject/WeakObjectPtrTemplates.h"
 
-namespace Interchange
+namespace UE
 {
-
-class FTaskPipeline
-{
-private:
-	TWeakObjectPtr<UInterchangePipelineBase> PipelineBase;
-	TWeakPtr<Interchange::FImportAsyncHelper, ESPMode::ThreadSafe> WeakAsyncHelper;
-public:
-	FTaskPipeline(TWeakObjectPtr<UInterchangePipelineBase> InPipelineBase, TWeakPtr<Interchange::FImportAsyncHelper, ESPMode::ThreadSafe> InAsyncHelper)
-		: PipelineBase(InPipelineBase)
-		, WeakAsyncHelper(InAsyncHelper)
-	{
-	}
-
-	ENamedThreads::Type GetDesiredThread()
-	{
-		return PipelineBase.Get()->ScriptedCanExecuteOnAnyThread() ? ENamedThreads::AnyBackgroundThreadNormalTask : ENamedThreads::GameThread;
-	}
-	
-	static ESubsequentsMode::Type GetSubsequentsMode()
+	namespace Interchange
 	{
-		return ESubsequentsMode::TrackSubsequents;
-	}
-
-	FORCEINLINE TStatId GetStatId() const
-	{
-		RETURN_QUICK_DECLARE_CYCLE_STAT(FTaskPipeline, STATGROUP_TaskGraphTasks);
-	}
-
-	void DoTask(ENamedThreads::Type CurrentThread, const FGraphEventRef& MyCompletionGraphEvent);
-};
-
 
-} // End namespace Interchange
+		class FTaskPipeline
+		{
+		private:
+			TWeakObjectPtr<UInterchangePipelineBase> PipelineBase;
+			TWeakPtr<FImportAsyncHelper, ESPMode::ThreadSafe> WeakAsyncHelper;
+		public:
+			FTaskPipeline(TWeakObjectPtr<UInterchangePipelineBase> InPipelineBase, TWeakPtr<FImportAsyncHelper, ESPMode::ThreadSafe> InAsyncHelper)
+				: PipelineBase(InPipelineBase)
+				, WeakAsyncHelper(InAsyncHelper)
+			{
+			}
+
+			ENamedThreads::Type GetDesiredThread()
+			{
+				return PipelineBase.Get()->ScriptedCanExecuteOnAnyThread() ? ENamedThreads::AnyBackgroundThreadNormalTask : ENamedThreads::GameThread;
+			}
+
+			static ESubsequentsMode::Type GetSubsequentsMode()
+			{
+				return ESubsequentsMode::TrackSubsequents;
+			}
+
+			FORCEINLINE TStatId GetStatId() const
+			{
+				RETURN_QUICK_DECLARE_CYCLE_STAT(FTaskPipeline, STATGROUP_TaskGraphTasks);
+			}
+
+			void DoTask(ENamedThreads::Type CurrentThread, const FGraphEventRef& MyCompletionGraphEvent);
+		};
+
+
+	} //ns Interchange
+}//ns UE
diff --git a/Engine/Source/Runtime/Interchange/Core/Private/Tasks/InterchangeTaskTranslator.cpp b/Engine/Source/Runtime/Interchange/Core/Private/Tasks/InterchangeTaskTranslator.cpp
--- a/Engine/Source/Runtime/Interchange/Core/Private/Tasks/InterchangeTaskTranslator.cpp
+++ b/Engine/Source/Runtime/Interchange/Core/Private/Tasks/InterchangeTaskTranslator.cpp
@@ -11,9 +11,9 @@
 #include "UObject/WeakObjectPtrTemplates.h"
 #include "Nodes/InterchangeBaseNodeContainer.h"
 
-void Interchange::FTaskTranslator::DoTask(ENamedThreads::Type CurrentThread, const FGraphEventRef& MyCompletionGraphEvent)
+void UE::Interchange::FTaskTranslator::DoTask(ENamedThreads::Type CurrentThread, const FGraphEventRef& MyCompletionGraphEvent)
 {
-	TSharedPtr<Interchange::FImportAsyncHelper, ESPMode::ThreadSafe> AsyncHelper = WeakAsyncHelper.Pin();
+	TSharedPtr<FImportAsyncHelper, ESPMode::ThreadSafe> AsyncHelper = WeakAsyncHelper.Pin();
 	check(AsyncHelper.IsValid());
 
 	UInterchangeTranslatorBase* Translator = AsyncHelper->Translators.IsValidIndex(SourceIndex) ? AsyncHelper->Translators[SourceIndex] : nullptr;
diff --git a/Engine/Source/Runtime/Interchange/Core/Private/Tasks/InterchangeTaskTranslator.h b/Engine/Source/Runtime/Interchange/Core/Private/Tasks/InterchangeTaskTranslator.h
--- a/Engine/Source/Runtime/Interchange/Core/Private/Tasks/InterchangeTaskTranslator.h
+++ b/Engine/Source/Runtime/Interchange/Core/Private/Tasks/InterchangeTaskTranslator.h
@@ -8,38 +8,41 @@
 #include "Stats/Stats.h"
 #include "UObject/WeakObjectPtrTemplates.h"
 
-namespace Interchange
+namespace UE
 {
-
-class FTaskTranslator
-{
-private:
-	int32 SourceIndex = INDEX_NONE;
-	TWeakPtr<Interchange::FImportAsyncHelper, ESPMode::ThreadSafe> WeakAsyncHelper;
-
-public:
-	FTaskTranslator(int32 InSourceIndex, TWeakPtr<Interchange::FImportAsyncHelper, ESPMode::ThreadSafe> InAsyncHelper)
-		: SourceIndex(InSourceIndex)
-		, WeakAsyncHelper(InAsyncHelper)
-	{
-	}
-
-	static FORCEINLINE ENamedThreads::Type GetDesiredThread()
-	{
-		return ENamedThreads::AnyBackgroundThreadNormalTask;
-	}
-	static FORCEINLINE ESubsequentsMode::Type GetSubsequentsMode()
+	namespace Interchange
 	{
-		return ESubsequentsMode::TrackSubsequents;
-	}
-
-	FORCEINLINE TStatId GetStatId() const
-	{
-		RETURN_QUICK_DECLARE_CYCLE_STAT(FTaskTranslator, STATGROUP_TaskGraphTasks);
-	}
-
-	void DoTask(ENamedThreads::Type CurrentThread, const FGraphEventRef& MyCompletionGraphEvent);
-};
-
 
-} // End namespace Interchange
+		class FTaskTranslator
+		{
+		private:
+			int32 SourceIndex = INDEX_NONE;
+			TWeakPtr<FImportAsyncHelper, ESPMode::ThreadSafe> WeakAsyncHelper;
+
+		public:
+			FTaskTranslator(int32 InSourceIndex, TWeakPtr<FImportAsyncHelper, ESPMode::ThreadSafe> InAsyncHelper)
+				: SourceIndex(InSourceIndex)
+				, WeakAsyncHelper(InAsyncHelper)
+			{
+			}
+
+			static FORCEINLINE ENamedThreads::Type GetDesiredThread()
+			{
+				return ENamedThreads::AnyBackgroundThreadNormalTask;
+			}
+			static FORCEINLINE ESubsequentsMode::Type GetSubsequentsMode()
+			{
+				return ESubsequentsMode::TrackSubsequents;
+			}
+
+			FORCEINLINE TStatId GetStatId() const
+			{
+				RETURN_QUICK_DECLARE_CYCLE_STAT(FTaskTranslator, STATGROUP_TaskGraphTasks);
+			}
+
+			void DoTask(ENamedThreads::Type CurrentThread, const FGraphEventRef& MyCompletionGraphEvent);
+		};
+
+
+	} //ns Interchange
+} //ns UE
diff --git a/Engine/Source/Runtime/Interchange/Core/Private/Tests/StorageTest.cpp b/Engine/Source/Runtime/Interchange/Core/Private/Tests/StorageTest.cpp
--- a/Engine/Source/Runtime/Interchange/Core/Private/Tests/StorageTest.cpp
+++ b/Engine/Source/Runtime/Interchange/Core/Private/Tests/StorageTest.cpp
@@ -14,7 +14,7 @@ IMPLEMENT_SIMPLE_AUTOMATION_TEST(FAttributeStorageTest, "System.Runtime.Intercha
 
 bool FAttributeStorageTest::RunTest(const FString& Parameters)
 {
-	using namespace Interchange;
+	using namespace UE::Interchange;
 
 	FString KeyPrefixString = TEXT("TestKey");
 	FName KeyPrefix(*KeyPrefixString);
@@ -118,7 +118,7 @@ bool FAttributeStorageTest::RunTest(const FString& Parameters)
 			ValueRef.Add(static_cast<uint8>(RandomStream.RandHelper(RandHelperUInt8)));
 		}
 
-		if (!IsAttributeStorageResultSuccess(TestStorage.RegisterAttribute(BigArrayKey, ValueRef, Interchange::EAttributeProperty::NoHash)))
+		if (!IsAttributeStorageResultSuccess(TestStorage.RegisterAttribute(BigArrayKey, ValueRef, EAttributeProperty::NoHash)))
 		{
 			AddError(FString(TEXT("`AttributeStorage` must handle adding TArray<uint8> attribute")));
 		}
