Fix Material factory to use dependencies instead of child node to find the dependent texture. dependencies has nothing to do with child node.

diff --git a/Engine/Plugins/Runtime/Interchange/Source/Import/Private/Material/InterchangeMaterialFactory.cpp b/Engine/Plugins/Runtime/Interchange/Source/Import/Private/Material/InterchangeMaterialFactory.cpp
--- a/Engine/Plugins/Runtime/Interchange/Source/Import/Private/Material/InterchangeMaterialFactory.cpp
+++ b/Engine/Plugins/Runtime/Interchange/Source/Import/Private/Material/InterchangeMaterialFactory.cpp
@@ -139,14 +139,16 @@ UObject* UInterchangeMaterialFactory::CreateAsset(const UInterchangeMaterialFact
 						if (Arguments.NodeContainer)
 						{
 							FName NodeUniqueID = MaterialNode->GetUniqueID();
-							int32 ChildCount = Arguments.NodeContainer->GetNodeChildrenCount(NodeUniqueID);
-							for (int32 ChildIndex = 0; ChildIndex < ChildCount; ++ChildIndex)
+							TArray<FName> MaterialDependencies;
+							MaterialNode->GetDependecies(MaterialDependencies);
+							int32 DependenciesCount = MaterialDependencies.Num();
+							for (int32 DependIndex = 0; DependIndex < DependenciesCount; ++DependIndex)
 							{
-								const UInterchangeBaseNode* ChildNode = Arguments.NodeContainer->GetNodeChildren(NodeUniqueID, ChildIndex);
-								if (ChildNode->GetUniqueID() == OutTextureUID && ChildNode->ReferenceObject.IsAsset())
+								const UInterchangeBaseNode* DepNode = Arguments.NodeContainer->GetNode(MaterialDependencies[DependIndex]);
+								if (DepNode->GetUniqueID() == OutTextureUID && DepNode->ReferenceObject.IsAsset())
 								{
 									//Use Resolve object so we just look for in memory UObject
-									UObject* TextureObject = ChildNode->ReferenceObject.ResolveObject();
+									UObject* TextureObject = DepNode->ReferenceObject.ResolveObject();
 									if (TextureObject)
 									{
 										TextureReference = Cast<UTexture>(TextureObject);
