Fix static analysis

diff --git a/Engine/Plugins/Runtime/Interchange/Source/Parsers/Fbx/Private/FbxScene.cpp b/Engine/Plugins/Runtime/Interchange/Source/Parsers/Fbx/Private/FbxScene.cpp
--- a/Engine/Plugins/Runtime/Interchange/Source/Parsers/Fbx/Private/FbxScene.cpp
+++ b/Engine/Plugins/Runtime/Interchange/Source/Parsers/Fbx/Private/FbxScene.cpp
@@ -21,70 +21,73 @@ namespace UE
 				{
 					FString NodeName = FFbxConvert::MakeString(FFbxConvert::MakeName(Node->GetName()));
 					UInterchangeSceneNode* UnrealNode = FFbxScene::CreateTransformNode(NodeContainer, NodeName, JSonErrorMessages);
-					NodeContainer.AddNode(UnrealNode);
-					FTransform LocalTransform = FFbxConvert::GetTransform(Node->EvaluateLocalTransform());
-					UnrealNode->SetCustomLocalTransform(LocalTransform);
-					FTransform GlobalTransform = FFbxConvert::GetTransform(Node->EvaluateGlobalTransform());
-					UnrealNode->SetCustomGlobalTransform(GlobalTransform);
-
-					int32 AttributeCount = Node->GetNodeAttributeCount();
-					for (int32 AttributeIndex = 0; AttributeIndex < AttributeCount; ++AttributeIndex)
+					if (UnrealNode)
 					{
-						const FbxNodeAttribute* NodeAttribute = Node->GetNodeAttributeByIndex(AttributeIndex);
-						switch (NodeAttribute->GetAttributeType())
+						NodeContainer.AddNode(UnrealNode);
+						FTransform LocalTransform = FFbxConvert::GetTransform(Node->EvaluateLocalTransform());
+						UnrealNode->SetCustomLocalTransform(LocalTransform);
+						FTransform GlobalTransform = FFbxConvert::GetTransform(Node->EvaluateGlobalTransform());
+						UnrealNode->SetCustomGlobalTransform(GlobalTransform);
+
+						int32 AttributeCount = Node->GetNodeAttributeCount();
+						for (int32 AttributeIndex = 0; AttributeIndex < AttributeCount; ++AttributeIndex)
 						{
-							case FbxNodeAttribute::eLODGroup:
-								//Add LOD group node
-								//UnrealNode = FFbxScene::CreateLodGroupNode(NodeContainer, NodeName, JSonErrorMessages);
-								break;
-							case FbxNodeAttribute::eMesh:
-								//Add a mesh node
-								if (Node->GetMesh()->GetDeformerCount(FbxDeformer::eSkin) > 0)
-								{
-									//Skeletal mesh node
-									//UnrealNode = FFbxSkeletalMesh::CreateSkeletalMeshNode(NodeContainer, NodeName, JSonErrorMessages);
-								}
-								else
-								{
-									//Static mesh node
-									//UnrealNode = FFbxStaticMesh::CreateStaticMeshNode(NodeContainer, NodeName, JSonErrorMessages);
-								}
-								break;
-							case FbxNodeAttribute::eSkeleton:
-								//Add a joint node
-								//UnrealNode = FFbxSkeletalMesh::CreateJointNode(NodeContainer, NodeName, JSonErrorMessages);
-								break;
-							case FbxNodeAttribute::eCamera:
-								//Add a camera node
-								//UnrealNode = FFbxCamera::CreateCameraNode(NodeContainer, NodeName, JSonErrorMessages);
-								break;
-							case FbxNodeAttribute::eCameraSwitcher:
-								//Add a camera switcher node
-								//UnrealNode = FFbxCamera::CreateCameraSwitcherNode(NodeContainer, NodeName, JSonErrorMessages);
-								break;
-							case FbxNodeAttribute::eLight:
-								//Add a light node
-								//UnrealNode = FFbxLight::CreateLightNode(NodeContainer, NodeName, JSonErrorMessages);
-								break;
-							default:
-								//Add a transform node
-								UnrealNode = FFbxScene::CreateTransformNode(NodeContainer, NodeName, JSonErrorMessages);
-								break;
+							const FbxNodeAttribute* NodeAttribute = Node->GetNodeAttributeByIndex(AttributeIndex);
+							switch (NodeAttribute->GetAttributeType())
+							{
+								case FbxNodeAttribute::eLODGroup:
+									//Add LOD group node
+									//UnrealNode = FFbxScene::CreateLodGroupNode(NodeContainer, NodeName, JSonErrorMessages);
+									break;
+								case FbxNodeAttribute::eMesh:
+									//Add a mesh node
+									if (Node->GetMesh()->GetDeformerCount(FbxDeformer::eSkin) > 0)
+									{
+										//Skeletal mesh node
+										//UnrealNode = FFbxSkeletalMesh::CreateSkeletalMeshNode(NodeContainer, NodeName, JSonErrorMessages);
+									}
+									else
+									{
+										//Static mesh node
+										//UnrealNode = FFbxStaticMesh::CreateStaticMeshNode(NodeContainer, NodeName, JSonErrorMessages);
+									}
+									break;
+								case FbxNodeAttribute::eSkeleton:
+									//Add a joint node
+									//UnrealNode = FFbxSkeletalMesh::CreateJointNode(NodeContainer, NodeName, JSonErrorMessages);
+									break;
+								case FbxNodeAttribute::eCamera:
+									//Add a camera node
+									//UnrealNode = FFbxCamera::CreateCameraNode(NodeContainer, NodeName, JSonErrorMessages);
+									break;
+								case FbxNodeAttribute::eCameraSwitcher:
+									//Add a camera switcher node
+									//UnrealNode = FFbxCamera::CreateCameraSwitcherNode(NodeContainer, NodeName, JSonErrorMessages);
+									break;
+								case FbxNodeAttribute::eLight:
+									//Add a light node
+									//UnrealNode = FFbxLight::CreateLightNode(NodeContainer, NodeName, JSonErrorMessages);
+									break;
+								default:
+									//Add a transform node
+									UnrealNode = FFbxScene::CreateTransformNode(NodeContainer, NodeName, JSonErrorMessages);
+									break;
+							}
 						}
-					}
-					if (UnrealNode)
-					{
 						NodeContainer.AddNode(UnrealNode);
 						//Add the dependencies of the material on the correct order
 						FFbxMaterial::AddAllNodeMaterials(UnrealNode, Node, NodeContainer, JSonErrorMessages);
-					}
 
-					const int32 ChildCount = Node->GetChildCount();
-					for (int32 ChildIndex = 0; ChildIndex < ChildCount; ++ChildIndex)
-					{
-						FbxNode* ChildNode = Node->GetChild(ChildIndex);
-						UInterchangeSceneNode* UnrealChildNode = AddHierarchyRecursively(ChildNode, SDKScene, NodeContainer, JSonErrorMessages);
-						UnrealChildNode->SetParentUID(UnrealNode->GetUniqueID());
+						const int32 ChildCount = Node->GetChildCount();
+						for (int32 ChildIndex = 0; ChildIndex < ChildCount; ++ChildIndex)
+						{
+							FbxNode* ChildNode = Node->GetChild(ChildIndex);
+							UInterchangeSceneNode* UnrealChildNode = AddHierarchyRecursively(ChildNode, SDKScene, NodeContainer, JSonErrorMessages);
+							if (UnrealChildNode)
+							{
+								UnrealChildNode->SetParentUID(UnrealNode->GetUniqueID());
+							}
+						}
 					}
 					return UnrealNode;
 				}
