Fix Linux CIS to enable exception for the fbx parser

diff --git a/Engine/Plugins/Runtime/Interchange/Source/Parsers/Fbx/InterchangeFbxParser.Build.cs b/Engine/Plugins/Runtime/Interchange/Source/Parsers/Fbx/InterchangeFbxParser.Build.cs
--- a/Engine/Plugins/Runtime/Interchange/Source/Parsers/Fbx/InterchangeFbxParser.Build.cs
+++ b/Engine/Plugins/Runtime/Interchange/Source/Parsers/Fbx/InterchangeFbxParser.Build.cs
@@ -8,7 +8,7 @@ namespace UnrealBuildTool.Rules
 	{
 		public InterchangeFbxParser(ReadOnlyTargetRules Target) : base(Target)
 		{
-			//bEnableExceptions = true;
+			bEnableExceptions = true;
 			
             PublicDependencyModuleNames.AddRange(
 				new string[]
