Fix CIS non unity build

diff --git a/Engine/Plugins/Runtime/Interchange/Source/Export/Private/InterchangeExportPlugin.cpp b/Engine/Plugins/Runtime/Interchange/Source/Export/Private/InterchangeExportPlugin.cpp
--- a/Engine/Plugins/Runtime/Interchange/Source/Export/Private/InterchangeExportPlugin.cpp
+++ b/Engine/Plugins/Runtime/Interchange/Source/Export/Private/InterchangeExportPlugin.cpp
@@ -1,11 +1,13 @@
 // Copyright Epic Games, Inc. All Rights Reserved.
 
 #include "CoreMinimal.h"
-#include "Modules/ModuleManager.h"
+#include "Engine/Engine.h"
 #include "IInterchangeExportPlugin.h"
 #include "InterchangeManager.h"
 #include "InterchangeTextureWriter.h"
-#include "Engine/Engine.h"
+#include "Misc/CoreDelegates.h"
+#include "Modules/ModuleManager.h"
+
 
 class FInterchangeExportPlugin : public IInterchangeExportPlugin
 {
