Interchange Fbx skeletal mesh support phase 1

diff --git a/Engine/Plugins/Runtime/Interchange/Source/Dispatcher/Private/InterchangeDispatcher.cpp b/Engine/Plugins/Runtime/Interchange/Source/Dispatcher/Private/InterchangeDispatcher.cpp
--- a/Engine/Plugins/Runtime/Interchange/Source/Dispatcher/Private/InterchangeDispatcher.cpp
+++ b/Engine/Plugins/Runtime/Interchange/Source/Dispatcher/Private/InterchangeDispatcher.cpp
@@ -103,11 +103,18 @@ namespace UE
 			SpawnHandler();
 		}
 
-		void FInterchangeDispatcher::StopProcess()
+		void FInterchangeDispatcher::StopProcess(bool bBlockUntilTerminated)
 		{
 			if (IsHandlerAlive())
 			{
-				WorkerHandler->Stop();
+				if (bBlockUntilTerminated)
+				{
+					WorkerHandler->StopBlocking();
+				}
+				else
+				{
+					WorkerHandler->Stop();
+				}
 			}
 		}
 
@@ -174,7 +181,7 @@ namespace UE
 
 		void FInterchangeDispatcher::CloseHandler()
 		{
-			StopProcess();
+			StopProcess(false);
 			WorkerHandler.Reset();
 		}
 
diff --git a/Engine/Plugins/Runtime/Interchange/Source/Dispatcher/Private/InterchangeDispatcherTask.cpp b/Engine/Plugins/Runtime/Interchange/Source/Dispatcher/Private/InterchangeDispatcherTask.cpp
--- a/Engine/Plugins/Runtime/Interchange/Source/Dispatcher/Private/InterchangeDispatcherTask.cpp
+++ b/Engine/Plugins/Runtime/Interchange/Source/Dispatcher/Private/InterchangeDispatcherTask.cpp
@@ -21,10 +21,10 @@ namespace UE
 			TSharedPtr<FJsonObject> CmdObject = MakeShared<FJsonObject>();
 			TSharedPtr<FJsonObject> ActionDataObject = MakeShared<FJsonObject>();
 			//CmdObject
-			CmdObject->SetStringField(GetCommandIDKey(), GetAction());
-			CmdObject->SetStringField(GetTranslatorIDKey(), GetTranslatorID());
-			ActionDataObject->SetStringField(GetSourceFilenameKey(), GetSourceFilename());
-			CmdObject->SetObjectField(GetCommandDataKey(), ActionDataObject);
+			CmdObject->SetStringField(GetCommandIDJsonKey(), GetAction());
+			CmdObject->SetStringField(GetTranslatorIDJsonKey(), GetTranslatorID());
+			ActionDataObject->SetStringField(GetSourceFilenameJsonKey(), GetSourceFilename());
+			CmdObject->SetObjectField(GetCommandDataJsonKey(), ActionDataObject);
 
 			FString LoadSourceCmd;
 			TSharedRef< TJsonWriter< TCHAR, TPrettyJsonPrintPolicy<TCHAR> > > JsonWriter = TJsonWriterFactory< TCHAR, TPrettyJsonPrintPolicy<TCHAR> >::Create(&LoadSourceCmd);
@@ -47,7 +47,7 @@ namespace UE
 				return false;
 			}
 			FString JsonActionValue;
-			if (!CmdObject->TryGetStringField(GetCommandIDKey(), JsonActionValue))
+			if (!CmdObject->TryGetStringField(GetCommandIDJsonKey(), JsonActionValue))
 			{
 				//The json cmd id key is missing
 				return false;
@@ -60,18 +60,18 @@ namespace UE
 			}
 
 			//Read the json
-			if (!CmdObject->TryGetStringField(GetTranslatorIDKey(), TranslatorID))
+			if (!CmdObject->TryGetStringField(GetTranslatorIDJsonKey(), TranslatorID))
 			{
 				//Missing Load command translator ID
 				return false;
 			}
 			const TSharedPtr<FJsonObject>* ActionDataObject = nullptr;
-			if (!CmdObject->TryGetObjectField(GetCommandDataKey(), ActionDataObject))
+			if (!CmdObject->TryGetObjectField(GetCommandDataJsonKey(), ActionDataObject))
 			{
 				//Missing Load Action data object
 				return false;
 			}
-			if (!((*ActionDataObject)->TryGetStringField(GetSourceFilenameKey(), SourceFilename)))
+			if (!((*ActionDataObject)->TryGetStringField(GetSourceFilenameJsonKey(), SourceFilename)))
 			{
 				return false;
 			}
@@ -85,7 +85,7 @@ namespace UE
 		{
 			TSharedPtr<FJsonObject> ResultObject = MakeShared<FJsonObject>();
 			//CmdObject
-			ResultObject->SetStringField(GetResultFilenameKey(), GetResultFilename());
+			ResultObject->SetStringField(GetResultFilenameJsonKey(), GetResultFilename());
 
 			FString JsonResult;
 			TSharedRef< TJsonWriter< TCHAR, TPrettyJsonPrintPolicy<TCHAR> > > JsonWriter = TJsonWriterFactory< TCHAR, TPrettyJsonPrintPolicy<TCHAR> >::Create(&JsonResult);
@@ -107,8 +107,104 @@ namespace UE
 				//Cannot read the json file
 				return false;
 			}
-			return ResultObject->TryGetStringField(GetResultFilenameKey(), ResultFilename);
+			return ResultObject->TryGetStringField(GetResultFilenameJsonKey(), ResultFilename);
 		}
 
+		FString FJsonFetchPayloadCmd::ToJson() const
+		{
+			//Code should not do a ToJson if the data was not set before
+			ensure(bIsDataInitialize);
+
+			TSharedPtr<FJsonObject> CmdObject = MakeShared<FJsonObject>();
+			TSharedPtr<FJsonObject> ActionDataObject = MakeShared<FJsonObject>();
+			//CmdObject
+			CmdObject->SetStringField(GetCommandIDJsonKey(), GetAction());
+			CmdObject->SetStringField(GetTranslatorIDJsonKey(), GetTranslatorID());
+			ActionDataObject->SetStringField(GetPayloadKeyJsonKey(), GetPayloadKey());
+			CmdObject->SetObjectField(GetCommandDataJsonKey(), ActionDataObject);
+
+			FString FetchPayloadCmd;
+			TSharedRef< TJsonWriter< TCHAR, TPrettyJsonPrintPolicy<TCHAR> > > JsonWriter = TJsonWriterFactory< TCHAR, TPrettyJsonPrintPolicy<TCHAR> >::Create(&FetchPayloadCmd);
+			if (!FJsonSerializer::Serialize(CmdObject.ToSharedRef(), JsonWriter))
+			{
+				//Error creating the json cmd string 
+				return FString();
+			}
+			return FetchPayloadCmd;
+		}
+
+		bool FJsonFetchPayloadCmd::FromJson(const FString& JsonString)
+		{
+			TSharedRef<TJsonReader<TCHAR>> Reader = FJsonStringReader::Create(JsonString);
+
+			TSharedPtr<FJsonObject> CmdObject;
+			if (!FJsonSerializer::Deserialize(Reader, CmdObject) || !CmdObject.IsValid())
+			{
+				//Cannot read the json file
+				return false;
+			}
+			FString JsonActionValue;
+			if (!CmdObject->TryGetStringField(GetCommandIDJsonKey(), JsonActionValue))
+			{
+				//The json cmd id key is missing
+				return false;
+			}
+
+			if (JsonActionValue != GetAction())
+			{
+				//This json do not represent a load command
+				return false;
+			}
+
+			//Read the json
+			if (!CmdObject->TryGetStringField(GetTranslatorIDJsonKey(), TranslatorID))
+			{
+				//Missing Load command translator ID
+				return false;
+			}
+			const TSharedPtr<FJsonObject>* ActionDataObject = nullptr;
+			if (!CmdObject->TryGetObjectField(GetCommandDataJsonKey(), ActionDataObject))
+			{
+				//Missing Load Action data object
+				return false;
+			}
+			if (!((*ActionDataObject)->TryGetStringField(GetPayloadKeyJsonKey(), PayloadKey)))
+			{
+				return false;
+			}
+
+			//Since we filled the data from the json file, set the data has been initialize.
+			bIsDataInitialize = true;
+			return true;
+		}
+
+		FString FJsonFetchPayloadCmd::JsonResultParser::ToJson() const
+		{
+			TSharedPtr<FJsonObject> ResultObject = MakeShared<FJsonObject>();
+			//CmdObject
+			ResultObject->SetStringField(GetResultFilenameJsonKey(), GetResultFilename());
+
+			FString JsonResult;
+			TSharedRef< TJsonWriter< TCHAR, TPrettyJsonPrintPolicy<TCHAR> > > JsonWriter = TJsonWriterFactory< TCHAR, TPrettyJsonPrintPolicy<TCHAR> >::Create(&JsonResult);
+			if (!FJsonSerializer::Serialize(ResultObject.ToSharedRef(), JsonWriter))
+			{
+				//Error creating the json cmd string 
+				return FString();
+			}
+			return JsonResult;
+		}
+
+		bool FJsonFetchPayloadCmd::JsonResultParser::FromJson(const FString& JsonString)
+		{
+			TSharedRef<TJsonReader<TCHAR>> Reader = FJsonStringReader::Create(JsonString);
+
+			TSharedPtr<FJsonObject> ResultObject;
+			if (!FJsonSerializer::Deserialize(Reader, ResultObject) || !ResultObject.IsValid())
+			{
+				//Cannot read the json file
+				return false;
+			}
+			return ResultObject->TryGetStringField(GetResultFilenameJsonKey(), ResultFilename);
+		}
 	} //ns Interchange
 }//ns UE
diff --git a/Engine/Plugins/Runtime/Interchange/Source/Dispatcher/Private/InterchangeWorkerHandler.cpp b/Engine/Plugins/Runtime/Interchange/Source/Dispatcher/Private/InterchangeWorkerHandler.cpp
--- a/Engine/Plugins/Runtime/Interchange/Source/Dispatcher/Private/InterchangeWorkerHandler.cpp
+++ b/Engine/Plugins/Runtime/Interchange/Source/Dispatcher/Private/InterchangeWorkerHandler.cpp
@@ -30,10 +30,10 @@ namespace UE
 #elif PLATFORM_WINDOWS
 #if UE_BUILD_DEBUG
 				Path = FPaths::Combine(Path, TEXT("Win64/InterchangeWorker-Win64-Debug.exe"));
-#else
+#else //#if UE_BUILD_DEBUG
 				Path = FPaths::Combine(Path, TEXT("Win64/InterchangeWorker.exe"));
-#endif
-#endif
+#endif //#else UE_BUILD_DEBUG
+#endif //#elif PLATFORM_WINDOWS
 				if (!FPaths::FileExists(Path))
 				{
 					UE_LOG(LogInterchangeDispatcher, Warning, TEXT("InterchangeWorker executable not found. Expected location: %s"), *FPaths::ConvertRelativePathToFull(ProcessorPath));
@@ -59,7 +59,7 @@ namespace UE
 
 		FInterchangeWorkerHandler::~FInterchangeWorkerHandler()
 		{
-			IOThread.Join();
+			StopBlocking();
 		}
 
 		void FInterchangeWorkerHandler::StartWorkerProcess()
@@ -291,6 +291,15 @@ namespace UE
 			bShouldTerminate = true;
 		}
 
+		void FInterchangeWorkerHandler::StopBlocking()
+		{
+			Stop();
+			if (IOThread.IsJoinable())
+			{
+				IOThread.Join();
+			}
+		}
+
 		bool FInterchangeWorkerHandler::IsAlive() const
 		{
 			return WorkerState != EWorkerState::Terminated;
diff --git a/Engine/Plugins/Runtime/Interchange/Source/Dispatcher/Public/InterchangeDispatcher.h b/Engine/Plugins/Runtime/Interchange/Source/Dispatcher/Public/InterchangeDispatcher.h
--- a/Engine/Plugins/Runtime/Interchange/Source/Dispatcher/Public/InterchangeDispatcher.h
+++ b/Engine/Plugins/Runtime/Interchange/Source/Dispatcher/Public/InterchangeDispatcher.h
@@ -25,7 +25,7 @@ namespace UE
 			void GetTaskState(int32 TaskIndex, ETaskState& TaskState, FString& JsonResult, TArray<FString>& JSonMessages);
 
 			void StartProcess();
-			void StopProcess();
+			void StopProcess(bool bBlockUntilTerminated);
 			void TerminateProcess();
 			void WaitAllTaskToCompleteExecution();
 			bool IsOver();
diff --git a/Engine/Plugins/Runtime/Interchange/Source/Dispatcher/Public/InterchangeDispatcherTask.h b/Engine/Plugins/Runtime/Interchange/Source/Dispatcher/Public/InterchangeDispatcherTask.h
--- a/Engine/Plugins/Runtime/Interchange/Source/Dispatcher/Public/InterchangeDispatcherTask.h
+++ b/Engine/Plugins/Runtime/Interchange/Source/Dispatcher/Public/InterchangeDispatcherTask.h
@@ -50,17 +50,17 @@ namespace UE
 			 */
 			virtual bool FromJson(const FString& JsonString) = 0;
 
-			static FString GetCommandIDKey()
+			static FString GetCommandIDJsonKey()
 			{
 				static const FString Key = TEXT("CmdID");
 				return Key;
 			}
-			static FString GetTranslatorIDKey()
+			static FString GetTranslatorIDJsonKey()
 			{
 				static const FString Key = TEXT("TranslatorID");
 				return Key;
 			}
-			static FString GetCommandDataKey()
+			static FString GetCommandDataJsonKey()
 			{
 				static const FString Key = TEXT("CmdData");
 				return Key;
@@ -109,7 +109,7 @@ namespace UE
 				return SourceFilename;
 			}
 
-			static FString GetSourceFilenameKey()
+			static FString GetSourceFilenameJsonKey()
 			{
 				static const FString Key = TEXT("SourceFile");
 				return Key;
@@ -132,7 +132,7 @@ namespace UE
 				FString ToJson() const;
 				bool FromJson(const FString& JsonString);
 
-				static FString GetResultFilenameKey()
+				static FString GetResultFilenameJsonKey()
 				{
 					const FString Key = TEXT("ResultFile");
 					return Key;
@@ -146,5 +146,80 @@ namespace UE
 			FString SourceFilename = FString();
 		};
 
+		class INTERCHANGEDISPATCHER_API FJsonFetchPayloadCmd : public IJsonCmdBase
+		{
+		public:
+			FJsonFetchPayloadCmd()
+			{
+				bIsDataInitialize = false;
+			}
+
+			FJsonFetchPayloadCmd(const FString& InTranslatorID, const FString& InPayloadKey)
+				: TranslatorID(InTranslatorID)
+				, PayloadKey(InPayloadKey)
+			{
+				bIsDataInitialize = true;
+			}
+
+			virtual FString GetAction() const override
+			{
+				static const FString LoadString = TEXT("Payload");
+				return LoadString;
+			}
+
+			virtual FString GetTranslatorID() const override
+			{
+				//Code should not do query data if the data was not set before
+				ensure(bIsDataInitialize);
+				return TranslatorID;
+			}
+
+			virtual FString ToJson() const;
+			virtual bool FromJson(const FString& JsonString);
+
+			FString GetPayloadKey() const
+			{
+				//Code should not do query data if the data was not set before
+				ensure(bIsDataInitialize);
+				return PayloadKey;
+			}
+
+			static FString GetPayloadKeyJsonKey()
+			{
+				static const FString Key = TEXT("PayloadKey");
+				return Key;
+			}
+
+			/**
+			 * Use this class helper to create the cmd result json string and to read it
+			 */
+			class INTERCHANGEDISPATCHER_API JsonResultParser
+			{
+			public:
+				FString GetResultFilename() const
+				{
+					return ResultFilename;
+				}
+				void SetResultFilename(const FString& InResultFilename)
+				{
+					ResultFilename = InResultFilename;
+				}
+				FString ToJson() const;
+				bool FromJson(const FString& JsonString);
+
+				static FString GetResultFilenameJsonKey()
+				{
+					const FString Key = TEXT("ResultFile");
+					return Key;
+				}
+			private:
+				FString ResultFilename = FString();
+			};
+
+		private:
+			FString TranslatorID = FString();
+			FString PayloadKey = FString();
+		};
+
 	} //ns Interchange
 }//ns UE
diff --git a/Engine/Plugins/Runtime/Interchange/Source/Dispatcher/Public/InterchangeWorkerHandler.h b/Engine/Plugins/Runtime/Interchange/Source/Dispatcher/Public/InterchangeWorkerHandler.h
--- a/Engine/Plugins/Runtime/Interchange/Source/Dispatcher/Public/InterchangeWorkerHandler.h
+++ b/Engine/Plugins/Runtime/Interchange/Source/Dispatcher/Public/InterchangeWorkerHandler.h
@@ -47,6 +47,7 @@ namespace UE
 			void Run();
 			bool IsAlive() const;
 			void Stop();
+			void StopBlocking();
 
 		private:
 			void RunInternal();
diff --git a/Engine/Plugins/Runtime/Interchange/Source/Import/InterchangeImportPlugin.Build.cs b/Engine/Plugins/Runtime/Interchange/Source/Import/InterchangeImportPlugin.Build.cs
--- a/Engine/Plugins/Runtime/Interchange/Source/Import/InterchangeImportPlugin.Build.cs
+++ b/Engine/Plugins/Runtime/Interchange/Source/Import/InterchangeImportPlugin.Build.cs
@@ -17,6 +17,8 @@ namespace UnrealBuildTool.Rules
 					"InterchangeEngine",
 					"InterchangeNodePlugin",
 					"MeshDescription",
+					"StaticMeshDescription",
+					"SkeletalMeshDescription",
 				}
 				);
 			PrivateDependencyModuleNames.AddRange(
@@ -28,6 +30,10 @@ namespace UnrealBuildTool.Rules
 					"Json",
 				}
 				);
+
+			OptimizeCode = CodeOptimization.Never;
+			bUseUnity = false;
+			PCHUsage = PCHUsageMode.NoPCHs;
 		}
 	}
 }
diff --git a/Engine/Plugins/Runtime/Interchange/Source/Import/Private/Fbx/InterchangeFbxTranslator.cpp b/Engine/Plugins/Runtime/Interchange/Source/Import/Private/Fbx/InterchangeFbxTranslator.cpp
--- a/Engine/Plugins/Runtime/Interchange/Source/Import/Private/Fbx/InterchangeFbxTranslator.cpp
+++ b/Engine/Plugins/Runtime/Interchange/Source/Import/Private/Fbx/InterchangeFbxTranslator.cpp
@@ -18,6 +18,7 @@
 #include "Nodes/InterchangeBaseNodeContainer.h"
 #include "Serialization/JsonWriter.h"
 #include "Serialization/JsonSerializer.h"
+#include "Serialization/LargeMemoryReader.h"
 #include "Texture/InterchangeTexturePayloadData.h"
 #include "UObject/GCObjectScopeGuard.h"
 
@@ -113,7 +114,8 @@ void UInterchangeFbxTranslator::ReleaseSource()
 {
 	if (Dispatcher.IsValid())
 	{
-		Dispatcher->StopProcess();
+		//Do not block the main thread
+		Dispatcher->StopProcess(!IsInGameThread());
 	}
 }
 
@@ -151,20 +153,76 @@ TOptional<UE::Interchange::FMaterialPayloadData> UInterchangeFbxTranslator::GetM
 	return TOptional<UE::Interchange::FMaterialPayloadData>();
 }
 
-TOptional<UE::Interchange::FStaticMeshPayloadData> UInterchangeFbxTranslator::GetStaticMeshPayloadData(const UInterchangeSourceData* SourceData, const FString& PayLoadKey) const
+TOptional<UE::Interchange::FStaticMeshPayloadData> UInterchangeFbxTranslator::GetStaticMeshPayloadData(const FString& PayLoadKey) const
 {
 	//Not implemented, currently we do not have any payload data for static meshes
 	return TOptional<UE::Interchange::FStaticMeshPayloadData>();
 }
 
-TOptional<UE::Interchange::FSkeletalMeshPayloadData> UInterchangeFbxTranslator::GetSkeletalMeshPayloadData(const UInterchangeSourceData* SourceData, const FString& PayLoadKey) const
+TOptional<UE::Interchange::FSkeletalMeshLodPayloadData> UInterchangeFbxTranslator::GetSkeletalMeshLodPayloadData(const FString& PayLoadKey) const
 {
-	//Not implemented, currently we do not have any payload data for skeletal meshes
-	return TOptional<UE::Interchange::FSkeletalMeshPayloadData>();
+	if (!Dispatcher.IsValid())
+	{
+		return TOptional<UE::Interchange::FSkeletalMeshLodPayloadData>();
+	}
+
+	//Create a json command to read the fbx file
+	FString JsonCommand = CreateFetchPayloadFbxCommand(PayLoadKey);
+	int32 TaskIndex = Dispatcher->AddTask(JsonCommand);
+
+	//Blocking call until all tasks are executed
+	Dispatcher->WaitAllTaskToCompleteExecution();
+
+	UE::Interchange::ETaskState TaskState;
+	FString JsonResult;
+	TArray<FString> JSonMessages;
+	Dispatcher->GetTaskState(TaskIndex, TaskState, JsonResult, JSonMessages);
+
+	//TODO: Parse the JSonMessage and add the message to the interchange not yet develop error messaging
+
+	if (TaskState != UE::Interchange::ETaskState::ProcessOk)
+	{
+		return TOptional<UE::Interchange::FSkeletalMeshLodPayloadData>();
+	}
+	//Grab the result file and fill the BaseNodeContainer
+	UE::Interchange::FJsonFetchPayloadCmd::JsonResultParser ResultParser;
+	ResultParser.FromJson(JsonResult);
+	FString SkeletalMeshPayloadFilename = ResultParser.GetResultFilename();
+
+	if (!ensure(FPaths::FileExists(SkeletalMeshPayloadFilename)))
+	{
+		//TODO log an error saying the payload file do not exist even if the get payload command succeed
+		return TOptional<UE::Interchange::FSkeletalMeshLodPayloadData>();
+	}
+	UE::Interchange::FSkeletalMeshLodPayloadData SkeletalMeshLodPayload;
+	SkeletalMeshLodPayload.LodMeshDescription.Empty();
+	//All sub object should be gone with the reset
+	TArray64<uint8> Buffer;
+	FFileHelper::LoadFileToArray(Buffer, *SkeletalMeshPayloadFilename);
+	uint8* FileData = Buffer.GetData();
+	int64 FileDataSize = Buffer.Num();
+	if (FileDataSize < 1)
+	{
+		//Nothing to load from this file
+		return TOptional<UE::Interchange::FSkeletalMeshLodPayloadData>();
+	}
+	//Buffer keep the ownership of the data, the large memory reader is use to serialize the TMap
+	FLargeMemoryReader Ar(FileData, FileDataSize);
+	SkeletalMeshLodPayload.LodMeshDescription.Serialize(Ar);
+
+	return SkeletalMeshLodPayload;
+
+	//return TOptional<UE::Interchange::FSkeletalMeshLodPayloadData>(SkeletalMeshLodPayload);
 }
 
 FString UInterchangeFbxTranslator::CreateLoadFbxFileCommand(const FString& FbxFilePath) const
 {
 	UE::Interchange::FJsonLoadSourceCmd LoadSourceCommand(TEXT("FBX"), FbxFilePath);
 	return LoadSourceCommand.ToJson();
+}
+
+FString UInterchangeFbxTranslator::CreateFetchPayloadFbxCommand(const FString& FbxPayloadKey) const
+{
+	UE::Interchange::FJsonFetchPayloadCmd PayloadCommand(TEXT("FBX"), FbxPayloadKey);
+	return PayloadCommand.ToJson();
 }
\ No newline at end of file
diff --git a/Engine/Plugins/Runtime/Interchange/Source/Import/Private/InterchangeImportPlugin.cpp b/Engine/Plugins/Runtime/Interchange/Source/Import/Private/InterchangeImportPlugin.cpp
--- a/Engine/Plugins/Runtime/Interchange/Source/Import/Private/InterchangeImportPlugin.cpp
+++ b/Engine/Plugins/Runtime/Interchange/Source/Import/Private/InterchangeImportPlugin.cpp
@@ -5,6 +5,8 @@
 #include "IInterchangeImportPlugin.h"
 #include "InterchangeManager.h"
 #include "Material/InterchangeMaterialFactory.h"
+#include "Mesh/InterchangeSkeletalMeshFactory.h"
+#include "Mesh/InterchangeSkeletonFactory.h"
 #include "Misc/CoreDelegates.h"
 #include "Modules/ModuleManager.h"
 #include "Texture/InterchangeBMPTranslator.h"
@@ -52,6 +54,8 @@ void FInterchangeImportPlugin::StartupModule()
 		//Register the factories
 		InterchangeManager.RegisterFactory(UInterchangeTextureFactory::StaticClass());
 		InterchangeManager.RegisterFactory(UInterchangeMaterialFactory::StaticClass());
+		InterchangeManager.RegisterFactory(UInterchangeSkeletonFactory::StaticClass());
+		InterchangeManager.RegisterFactory(UInterchangeSkeletalMeshFactory::StaticClass());
 	};
 
 	if (GEngine)
diff --git a/Engine/Plugins/Runtime/Interchange/Source/Import/Private/Material/InterchangeMaterialFactory.cpp b/Engine/Plugins/Runtime/Interchange/Source/Import/Private/Material/InterchangeMaterialFactory.cpp
--- a/Engine/Plugins/Runtime/Interchange/Source/Import/Private/Material/InterchangeMaterialFactory.cpp
+++ b/Engine/Plugins/Runtime/Interchange/Source/Import/Private/Material/InterchangeMaterialFactory.cpp
@@ -68,6 +68,7 @@ UObject* UInterchangeMaterialFactory::CreateEmptyAsset(const FCreateAssetParams&
 		UE_LOG(LogInterchangeImportPlugin, Warning, TEXT("Could not create Material asset %s"), *Arguments.AssetName);
 		return nullptr;
 	}
+	Material->PreEditChange(nullptr);
 #endif //WITH_EDITORONLY_DATA
 	return Material;
 }
@@ -149,7 +150,7 @@ UObject* UInterchangeMaterialFactory::CreateAsset(const UInterchangeMaterialFact
 							for (int32 DependIndex = 0; DependIndex < DependenciesCount; ++DependIndex)
 							{
 								const UInterchangeBaseNode* DepNode = Arguments.NodeContainer->GetNode(MaterialDependencies[DependIndex]);
-								if (DepNode->GetUniqueID() == OutTextureUID && DepNode->ReferenceObject.IsAsset())
+								if (DepNode && DepNode->GetUniqueID() == OutTextureUID && DepNode->ReferenceObject.IsAsset())
 								{
 									//Use Resolve object so we just look for in memory UObject
 									UObject* TextureObject = DepNode->ReferenceObject.ResolveObject();
diff --git a/Engine/Plugins/Runtime/Interchange/Source/Import/Private/Mesh/InterchangeSkeletalMeshFactory.cpp b/Engine/Plugins/Runtime/Interchange/Source/Import/Private/Mesh/InterchangeSkeletalMeshFactory.cpp
new file mode 100644
--- /dev/null
+++ b/Engine/Plugins/Runtime/Interchange/Source/Import/Private/Mesh/InterchangeSkeletalMeshFactory.cpp
@@ -0,0 +1,1252 @@
+// Copyright Epic Games, Inc. All Rights Reserved. 
+#include "Mesh/InterchangeSkeletalMeshFactory.h"
+
+#include "Components.h"
+#include "Engine/SkeletalMesh.h"
+#include "GPUSkinPublicDefs.h"
+#include "InterchangeImportCommon.h"
+#include "InterchangeJointNode.h"
+#include "InterchangeMaterialNode.h"
+#include "InterchangeSkeletalMeshLodDataNode.h"
+#include "InterchangeSkeletalMeshNode.h"
+#include "InterchangeSkeletonNode.h"
+#include "InterchangeSourceData.h"
+#include "InterchangeTranslatorBase.h"
+#include "LogInterchangeImportPlugin.h"
+#include "Mesh/InterchangeSkeletalMeshPayload.h"
+#include "Mesh/InterchangeSkeletalMeshPayloadInterface.h"
+#include "Nodes/InterchangeBaseNode.h"
+#include "Nodes/InterchangeBaseNodeContainer.h"
+#include "Rendering/SkeletalMeshLODImporterData.h"
+#include "Rendering/SkeletalMeshModel.h"
+#include "SkeletalMeshAttributes.h"
+#include "SkeletalMeshOperations.h"
+
+#if WITH_EDITORONLY_DATA
+
+#include "EditorFramework/AssetImportData.h"
+
+#endif //WITH_EDITORONLY_DATA
+
+#if WITH_EDITOR
+namespace UE
+{
+	namespace Interchange
+	{
+		namespace Private
+		{
+			struct FJointInfo
+			{
+				FName Name;
+				int32 ParentIndex;  // 0 if this is the root bone.  
+				FTransform	LocalTransform; // local transform
+			};
+
+			void RecursiveAddBones(const UInterchangeBaseNodeContainer* NodeContainer, FName JointNodeId, TArray <FJointInfo>& JointInfos, int32 ParentIndex, TArray<SkeletalMeshImportData::FBone>& RefBonesBinary)
+			{
+				const UInterchangeJointNode* JointNode = Cast<UInterchangeJointNode>(NodeContainer->GetNode(JointNodeId));
+				if (!JointNode)
+				{
+					UE_LOG(LogInterchangeImportPlugin, Warning, TEXT("Invalid Skeleton Joint"));
+					return;
+				}
+
+				int32 JointInfoIndex = JointInfos.Num();
+				FJointInfo& Info = JointInfos.AddZeroed_GetRef();
+				ensure(JointNode->GetCustomName(Info.Name));
+				ensure(JointNode->GetCustomLocalTransform(Info.LocalTransform));
+				Info.ParentIndex = ParentIndex;
+
+				SkeletalMeshImportData::FBone& Bone = RefBonesBinary.AddZeroed_GetRef();
+				Bone.Name = Info.Name.ToString();
+				Bone.BonePos.Transform = Info.LocalTransform;
+				Bone.ParentIndex = ParentIndex;
+				//Fill the scrap we do not need
+				Bone.BonePos.Length = 0.0f;
+				Bone.BonePos.XSize = 1.0f;
+				Bone.BonePos.YSize = 1.0f;
+				Bone.BonePos.ZSize = 1.0f;
+				
+				const TArray<FName> ChildrenIds = NodeContainer->GetNodeChildrenUIDs(JointNodeId);
+				Bone.NumChildren = ChildrenIds.Num();
+				for (int32 ChildIndex = 0; ChildIndex < ChildrenIds.Num(); ++ChildIndex)
+				{
+					RecursiveAddBones(NodeContainer, ChildrenIds[ChildIndex], JointInfos, JointInfoIndex, RefBonesBinary);
+				}
+			}
+
+			bool ProcessImportMeshSkeleton(const USkeleton* SkeletonAsset, FReferenceSkeleton& RefSkeleton, int32& SkeletalDepth, const UInterchangeBaseNodeContainer* NodeContainer, FName RootJointNodeId, TArray<SkeletalMeshImportData::FBone>& RefBonesBinary)
+			{
+				auto FixupBoneName = [](FString BoneName)
+				{
+					BoneName.TrimStartAndEndInline();
+					BoneName.ReplaceInline(TEXT(" "), TEXT("-"), ESearchCase::IgnoreCase);
+					return BoneName;
+				};
+
+				RefBonesBinary.Empty();
+				// Setup skeletal hierarchy + names structure.
+				RefSkeleton.Empty();
+
+				FReferenceSkeletonModifier RefSkelModifier(RefSkeleton, SkeletonAsset);
+				TArray <FJointInfo> JointInfos;
+				RecursiveAddBones(NodeContainer, RootJointNodeId, JointInfos, INDEX_NONE, RefBonesBinary);
+				// Digest bones to the serializable format.
+				for (int32 b = 0; b < JointInfos.Num(); b++)
+				{
+					const FJointInfo& BinaryBone = JointInfos[b];
+
+					const FString BoneName = FixupBoneName(BinaryBone.Name.ToString());
+					const FMeshBoneInfo BoneInfo(FName(*BoneName, FNAME_Add), BinaryBone.Name.ToString(), BinaryBone.ParentIndex);
+					const FTransform BoneTransform(BinaryBone.LocalTransform);
+					if (RefSkeleton.FindRawBoneIndex(BoneInfo.Name) != INDEX_NONE)
+					{
+						UE_LOG(LogInterchangeImportPlugin, Error, TEXT("Invalid Skeleton because of non-unique bone names [%s]"), *BoneInfo.Name.ToString());
+						return false;
+					}
+					RefSkelModifier.Add(BoneInfo, BoneTransform);
+				}
+
+				// Add hierarchy index to each bone and detect max depth.
+				SkeletalDepth = 0;
+
+				TArray<int32> SkeletalDepths;
+				SkeletalDepths.Empty(JointInfos.Num());
+				SkeletalDepths.AddZeroed(JointInfos.Num());
+				for (int32 BoneIndex = 0; BoneIndex < RefSkeleton.GetRawBoneNum(); ++BoneIndex)
+				{
+					int32 Parent = RefSkeleton.GetRawParentIndex(BoneIndex);
+					int32 Depth = 1.0f;
+
+					SkeletalDepths[BoneIndex] = 1.0f;
+					if (Parent != INDEX_NONE)
+					{
+						Depth += SkeletalDepths[Parent];
+					}
+					if (SkeletalDepth < Depth)
+					{
+						SkeletalDepth = Depth;
+					}
+					SkeletalDepths[BoneIndex] = Depth;
+				}
+
+				return true;
+			}
+
+			void CleanUpUnusedMaterials(FSkeletalMeshImportData& ImportData)
+			{
+				if (ImportData.Materials.Num() <= 0)
+				{
+					return;
+				}
+
+				TArray< SkeletalMeshImportData::FMaterial > ExistingMatList = ImportData.Materials;
+
+				TArray<uint8> UsedMaterialIndex;
+				// Find all material that are use by the mesh faces
+				int32 FaceNum = ImportData.Faces.Num();
+				for (int32 TriangleIndex = 0; TriangleIndex < FaceNum; TriangleIndex++)
+				{
+					SkeletalMeshImportData::FTriangle& Triangle = ImportData.Faces[TriangleIndex];
+					UsedMaterialIndex.AddUnique(Triangle.MatIndex);
+				}
+				//Remove any unused material.
+				if (UsedMaterialIndex.Num() < ExistingMatList.Num())
+				{
+					TArray<int32> RemapIndex;
+					TArray< SkeletalMeshImportData::FMaterial >& NewMatList = ImportData.Materials;
+					NewMatList.Empty();
+					for (int32 ExistingMatIndex = 0; ExistingMatIndex < ExistingMatList.Num(); ++ExistingMatIndex)
+					{
+						if (UsedMaterialIndex.Contains((uint8)ExistingMatIndex))
+						{
+							RemapIndex.Add(NewMatList.Add(ExistingMatList[ExistingMatIndex]));
+						}
+						else
+						{
+							RemapIndex.Add(INDEX_NONE);
+						}
+					}
+					ImportData.MaxMaterialIndex = 0;
+					//Remap the face material index
+					for (int32 TriangleIndex = 0; TriangleIndex < FaceNum; TriangleIndex++)
+					{
+						SkeletalMeshImportData::FTriangle& Triangle = ImportData.Faces[TriangleIndex];
+						check(RemapIndex[Triangle.MatIndex] != INDEX_NONE);
+						Triangle.MatIndex = RemapIndex[Triangle.MatIndex];
+						ImportData.MaxMaterialIndex = FMath::Max<uint32>(ImportData.MaxMaterialIndex, Triangle.MatIndex);
+					}
+				}
+			}
+
+			void SetMaterialSkinXXOrder(FSkeletalMeshImportData& ImportData)
+			{
+				TArray<int32> MaterialIndexToSkinIndex;
+				TMap<int32, int32> SkinIndexToMaterialIndex;
+				TArray<int32> MissingSkinSuffixMaterial;
+				TMap<int32, int32> SkinIndexGreaterThenMaterialArraySize;
+				{
+					int32 MaterialCount = ImportData.Materials.Num();
+
+					bool bNeedsReorder = false;
+					for (int32 MaterialIndex = 0; MaterialIndex < MaterialCount; ++MaterialIndex)
+					{
+						// get skin index
+						FString MatName = ImportData.Materials[MaterialIndex].MaterialImportName;
+
+						if (MatName.Len() > 6)
+						{
+							int32 Offset = MatName.Find(TEXT("_SKIN"), ESearchCase::IgnoreCase, ESearchDir::FromEnd);
+							if (Offset != INDEX_NONE)
+							{
+								// Chop off the material name so we are left with the number in _SKINXX
+								FString SkinXXNumber = MatName.Right(MatName.Len() - (Offset + 1)).RightChop(4);
+
+								if (SkinXXNumber.IsNumeric())
+								{
+									bNeedsReorder = true;
+									int32 TmpIndex = FPlatformString::Atoi(*SkinXXNumber);
+									if (TmpIndex < MaterialCount)
+									{
+										SkinIndexToMaterialIndex.Add(TmpIndex, MaterialIndex);
+									}
+									else
+									{
+										SkinIndexGreaterThenMaterialArraySize.Add(TmpIndex, MaterialIndex);
+									}
+
+								}
+							}
+							else
+							{
+								MissingSkinSuffixMaterial.Add(MaterialIndex);
+							}
+						}
+						else
+						{
+							MissingSkinSuffixMaterial.Add(MaterialIndex);
+						}
+					}
+
+					if (bNeedsReorder && MissingSkinSuffixMaterial.Num() > 0)
+					{
+						//TODO log a user warning message
+						//AddTokenizedErrorMessage(FTokenizedMessage::Create(EMessageSeverity::Error, LOCTEXT("FbxSkeletaLMeshimport_Skinxx_missing", "Cannot mix skinxx suffix materials with no skinxx material, mesh section order will not be right.")), FFbxErrors::Generic_Mesh_SkinxxNameError);
+						return;
+					}
+
+					//Add greater then material array skinxx at the end sorted by integer the index will be remap correctly in the case of a LOD import
+					if (SkinIndexGreaterThenMaterialArraySize.Num() > 0)
+					{
+						int32 MaxAvailableKey = SkinIndexToMaterialIndex.Num();
+						for (int32 AvailableKey = 0; AvailableKey < MaxAvailableKey; ++AvailableKey)
+						{
+							if (SkinIndexToMaterialIndex.Contains(AvailableKey))
+								continue;
+
+							TMap<int32, int32> TempSkinIndexToMaterialIndex;
+							for (auto KvpSkinToMat : SkinIndexToMaterialIndex)
+							{
+								if (KvpSkinToMat.Key > AvailableKey)
+								{
+									TempSkinIndexToMaterialIndex.Add(KvpSkinToMat.Key - 1, KvpSkinToMat.Value);
+								}
+								else
+								{
+									TempSkinIndexToMaterialIndex.Add(KvpSkinToMat.Key, KvpSkinToMat.Value);
+								}
+							}
+							//move all the later key of the array to fill the available index
+							SkinIndexToMaterialIndex = TempSkinIndexToMaterialIndex;
+							AvailableKey--; //We need to retest the same index it can be empty
+						}
+						//Reorder the array
+						SkinIndexGreaterThenMaterialArraySize.KeySort(TLess<int32>());
+						for (auto Kvp : SkinIndexGreaterThenMaterialArraySize)
+						{
+							SkinIndexToMaterialIndex.Add(SkinIndexToMaterialIndex.Num(), Kvp.Value);
+						}
+					}
+
+					//Fill the array MaterialIndexToSkinIndex so we order material by _skinXX order
+					//This ensure we support skinxx suffixe that are not increment by one like _skin00, skin_01, skin_03, skin_04, skin_08... 
+					for (auto kvp : SkinIndexToMaterialIndex)
+					{
+						int32 MatIndexToInsert = 0;
+						for (MatIndexToInsert = 0; MatIndexToInsert < MaterialIndexToSkinIndex.Num(); ++MatIndexToInsert)
+						{
+							if (*(SkinIndexToMaterialIndex.Find(MaterialIndexToSkinIndex[MatIndexToInsert])) >= kvp.Value)
+							{
+								break;
+							}
+						}
+						MaterialIndexToSkinIndex.Insert(kvp.Key, MatIndexToInsert);
+					}
+
+					if (bNeedsReorder)
+					{
+						// re-order the materials
+						TArray< SkeletalMeshImportData::FMaterial > ExistingMatList = ImportData.Materials;
+						for (int32 MissingIndex : MissingSkinSuffixMaterial)
+						{
+							MaterialIndexToSkinIndex.Insert(MaterialIndexToSkinIndex.Num(), MissingIndex);
+						}
+						for (int32 MaterialIndex = 0; MaterialIndex < MaterialCount; ++MaterialIndex)
+						{
+							if (MaterialIndex < MaterialIndexToSkinIndex.Num())
+							{
+								int32 NewIndex = MaterialIndexToSkinIndex[MaterialIndex];
+								if (ExistingMatList.IsValidIndex(NewIndex))
+								{
+									ImportData.Materials[NewIndex] = ExistingMatList[MaterialIndex];
+								}
+							}
+						}
+
+						// remapping the material index for each triangle
+						int32 FaceNum = ImportData.Faces.Num();
+						for (int32 TriangleIndex = 0; TriangleIndex < FaceNum; TriangleIndex++)
+						{
+							SkeletalMeshImportData::FTriangle& Triangle = ImportData.Faces[TriangleIndex];
+							if (Triangle.MatIndex < MaterialIndexToSkinIndex.Num())
+							{
+								Triangle.MatIndex = MaterialIndexToSkinIndex[Triangle.MatIndex];
+							}
+						}
+					}
+				}
+			}
+
+			namespace SmoothGroupHelper
+			{
+				struct tFaceRecord
+				{
+					int32 FaceIndex;
+					int32 HoekIndex;
+					int32 WedgeIndex;
+					uint32 SmoothFlags;
+					uint32 FanFlags;
+				};
+
+				struct VertsFans
+				{
+					TArray<tFaceRecord> FaceRecord;
+					int32 FanGroupCount;
+				};
+
+				struct tInfluences
+				{
+					TArray<int32> RawInfIndices;
+				};
+
+				struct tWedgeList
+				{
+					TArray<int32> WedgeList;
+				};
+
+				struct tFaceSet
+				{
+					TArray<int32> Faces;
+				};
+
+				// Check whether faces have at least two vertices in common. These must be POINTS - don't care about wedges.
+				bool FacesAreSmoothlyConnected(FSkeletalMeshImportData& ImportData, int32 Face1, int32 Face2)
+				{
+
+					//if( ( Face1 >= Thing->SkinData.Faces.Num()) || ( Face2 >= Thing->SkinData.Faces.Num()) ) return false;
+
+					if (Face1 == Face2)
+					{
+						return true;
+					}
+
+					// Smoothing groups match at least one bit in binary AND ?
+					if ((ImportData.Faces[Face1].SmoothingGroups & ImportData.Faces[Face2].SmoothingGroups) == 0)
+					{
+						return false;
+					}
+
+					int32 VertMatches = 0;
+					for (int32 i = 0; i < 3; i++)
+					{
+						int32 Point1 = ImportData.Wedges[ImportData.Faces[Face1].WedgeIndex[i]].VertexIndex;
+
+						for (int32 j = 0; j < 3; j++)
+						{
+							int32 Point2 = ImportData.Wedges[ImportData.Faces[Face2].WedgeIndex[j]].VertexIndex;
+							if (Point2 == Point1)
+							{
+								VertMatches++;
+							}
+						}
+					}
+
+					return (VertMatches >= 2);
+				}
+			} //namespace SmoothGroupHelper
+			int32 DoUnSmoothVerts(FSkeletalMeshImportData& ImportData)
+			{
+				//
+				// Connectivity: triangles with non-matching smoothing groups will be physically split.
+				//
+				// -> Splitting involves: the UV+material-contaning vertex AND the 3d point.
+				//
+				// -> Tally smoothing groups for each and every (textured) vertex.
+				//
+				// -> Collapse: 
+				// -> start from a vertex and all its adjacent triangles - go over
+				// each triangle - if any connecting one (sharing more than one vertex) gives a smoothing match,
+				// accumulate it. Then IF more than one resulting section, 
+				// ensure each boundary 'vert' is split _if not already_ to give each smoothing group
+				// independence from all others.
+				//
+
+				int32 DuplicatedVertCount = 0;
+				int32 RemappedHoeks = 0;
+
+				int32 TotalSmoothMatches = 0;
+				int32 TotalConnexChex = 0;
+
+				// Link _all_ faces to vertices.	
+				TArray<SmoothGroupHelper::VertsFans>  Fans;
+				TArray<SmoothGroupHelper::tInfluences> PointInfluences;
+				TArray<SmoothGroupHelper::tWedgeList>  PointWedges;
+
+				Fans.AddZeroed(ImportData.Points.Num());//Fans.AddExactZeroed(			Thing->SkinData.Points.Num() );
+				PointInfluences.AddZeroed(ImportData.Points.Num());//PointInfluences.AddExactZeroed( Thing->SkinData.Points.Num() );
+				PointWedges.AddZeroed(ImportData.Points.Num());//PointWedges.AddExactZeroed(	 Thing->SkinData.Points.Num() );
+
+				// Existing points map 1:1
+				ImportData.PointToRawMap.AddUninitialized(ImportData.Points.Num());
+				for (int32 i = 0; i < ImportData.Points.Num(); i++)
+				{
+					ImportData.PointToRawMap[i] = i;
+				}
+
+				for (int32 i = 0; i < ImportData.Influences.Num(); i++)
+				{
+					if (PointInfluences.Num() <= ImportData.Influences[i].VertexIndex)
+					{
+						PointInfluences.AddZeroed(ImportData.Influences[i].VertexIndex - PointInfluences.Num() + 1);
+					}
+					PointInfluences[ImportData.Influences[i].VertexIndex].RawInfIndices.Add(i);
+				}
+
+				for (int32 i = 0; i < ImportData.Wedges.Num(); i++)
+				{
+					if (uint32(PointWedges.Num()) <= ImportData.Wedges[i].VertexIndex)
+					{
+						PointWedges.AddZeroed(ImportData.Wedges[i].VertexIndex - PointWedges.Num() + 1);
+					}
+
+					PointWedges[ImportData.Wedges[i].VertexIndex].WedgeList.Add(i);
+				}
+
+				for (int32 f = 0; f < ImportData.Faces.Num(); f++)
+				{
+					// For each face, add a pointer to that face into the Fans[vertex].
+					for (int32 i = 0; i < 3; i++)
+					{
+						int32 WedgeIndex = ImportData.Faces[f].WedgeIndex[i];
+						int32 PointIndex = ImportData.Wedges[WedgeIndex].VertexIndex;
+						SmoothGroupHelper::tFaceRecord NewFR;
+
+						NewFR.FaceIndex = f;
+						NewFR.HoekIndex = i;
+						NewFR.WedgeIndex = WedgeIndex; // This face touches the point courtesy of Wedges[Wedgeindex].
+						NewFR.SmoothFlags = ImportData.Faces[f].SmoothingGroups;
+						NewFR.FanFlags = 0;
+						Fans[PointIndex].FaceRecord.Add(NewFR);
+						Fans[PointIndex].FanGroupCount = 0;
+					}
+				}
+
+				// Investigate connectivity and assign common group numbers (1..+) to the fans' individual FanFlags.
+				for (int32 p = 0; p < Fans.Num(); p++) // The fan of faces for each 3d point 'p'.
+				{
+					// All faces connecting.
+					if (Fans[p].FaceRecord.Num() > 0)
+					{
+						int32 FacesProcessed = 0;
+						TArray<SmoothGroupHelper::tFaceSet> FaceSets; // Sets with indices INTO FANS, not into face array.			
+
+						// Digest all faces connected to this vertex (p) into one or more smooth sets. only need to check 
+						// all faces MINUS one..
+						while (FacesProcessed < Fans[p].FaceRecord.Num())
+						{
+							// One loop per group. For the current ThisFaceIndex, tally all truly connected ones
+							// and put them in a new TArray. Once no more can be connected, stop.
+
+							int32 NewSetIndex = FaceSets.Num(); // 0 to start
+							FaceSets.AddZeroed(1);						// first one will be just ThisFaceIndex.
+
+							// Find the first non-processed face. There will be at least one.
+							int32 ThisFaceFanIndex = 0;
+							{
+								int32 SearchIndex = 0;
+								while (Fans[p].FaceRecord[SearchIndex].FanFlags == -1) // -1 indicates already  processed. 
+								{
+									SearchIndex++;
+								}
+								ThisFaceFanIndex = SearchIndex; //Fans[p].FaceRecord[SearchIndex].FaceIndex; 
+							}
+
+							// Initial face.
+							FaceSets[NewSetIndex].Faces.Add(ThisFaceFanIndex);   // Add the unprocessed Face index to the "local smoothing group" [NewSetIndex].
+							Fans[p].FaceRecord[ThisFaceFanIndex].FanFlags = -1;			  // Mark as processed.
+							FacesProcessed++;
+
+							// Find all faces connected to this face, and if there's any
+							// smoothing group matches, put it in current face set and mark it as processed;
+							// until no more match. 
+							int32 NewMatches = 0;
+							do
+							{
+								NewMatches = 0;
+								// Go over all current faces in this faceset and set if the FaceRecord (local smoothing groups) has any matches.
+								// there will be at least one face already in this faceset - the first face in the fan.
+								for (int32 n = 0; n < FaceSets[NewSetIndex].Faces.Num(); n++)
+								{
+									int32 HookFaceIdx = Fans[p].FaceRecord[FaceSets[NewSetIndex].Faces[n]].FaceIndex;
+
+									//Go over the fan looking for matches.
+									for (int32 s = 0; s < Fans[p].FaceRecord.Num(); s++)
+									{
+										// Skip if same face, skip if face already processed.
+										if ((HookFaceIdx != Fans[p].FaceRecord[s].FaceIndex) && (Fans[p].FaceRecord[s].FanFlags != -1))
+										{
+											TotalConnexChex++;
+											// Process if connected with more than one vertex, AND smooth..
+											if (SmoothGroupHelper::FacesAreSmoothlyConnected(ImportData, HookFaceIdx, Fans[p].FaceRecord[s].FaceIndex))
+											{
+												TotalSmoothMatches++;
+												Fans[p].FaceRecord[s].FanFlags = -1; // Mark as processed.
+												FacesProcessed++;
+												// Add 
+												FaceSets[NewSetIndex].Faces.Add(s); // Store FAN index of this face index into smoothing group's faces. 
+												// Tally
+												NewMatches++;
+											}
+										} // not the same...
+									}// all faces in fan
+								} // all faces in FaceSet
+							} while (NewMatches);
+
+						}// Repeat until all faces processed.
+
+						// For the new non-initialized  face sets, 
+						// Create a new point, influences, and uv-vertex(-ices) for all individual FanFlag groups with an index of 2+ and also remap
+						// the face's vertex into those new ones.
+						if (FaceSets.Num() > 1)
+						{
+							for (int32 f = 1; f < FaceSets.Num(); f++)
+							{
+								check(ImportData.Points.Num() == ImportData.PointToRawMap.Num());
+
+								// We duplicate the current vertex. (3d point)
+								int32 NewPointIndex = ImportData.Points.Num();
+								ImportData.Points.AddUninitialized();
+								ImportData.Points[NewPointIndex] = ImportData.Points[p];
+
+								ImportData.PointToRawMap.AddUninitialized();
+								ImportData.PointToRawMap[NewPointIndex] = p;
+
+								DuplicatedVertCount++;
+
+								// Duplicate all related weights.
+								for (int32 t = 0; t < PointInfluences[p].RawInfIndices.Num(); t++)
+								{
+									// Add new weight
+									int32 NewWeightIndex = ImportData.Influences.Num();
+									ImportData.Influences.AddUninitialized();
+									ImportData.Influences[NewWeightIndex] = ImportData.Influences[PointInfluences[p].RawInfIndices[t]];
+									ImportData.Influences[NewWeightIndex].VertexIndex = NewPointIndex;
+								}
+
+								// Duplicate any and all Wedges associated with it; and all Faces' wedges involved.					
+								for (int32 w = 0; w < PointWedges[p].WedgeList.Num(); w++)
+								{
+									int32 OldWedgeIndex = PointWedges[p].WedgeList[w];
+									int32 NewWedgeIndex = ImportData.Wedges.Num();
+									ImportData.Wedges[OldWedgeIndex].VertexIndex = NewPointIndex;
+								}
+							}
+						} //  if FaceSets.Num(). -> duplicate stuff
+					}//	while( FacesProcessed < Fans[p].FaceRecord.Num() )
+				} // Fans for each 3d point
+
+				if (!ensure(ImportData.Points.Num() == ImportData.PointToRawMap.Num()))
+				{
+					//TODO log a warning to the user
+
+					//Create a valid PointtoRawMap but with bad content
+					int32 PointNum = ImportData.Points.Num();
+					ImportData.PointToRawMap.Empty(PointNum);
+					for (int32 PointIndex = 0; PointIndex < PointNum; ++PointIndex)
+					{
+						ImportData.PointToRawMap[PointIndex] = PointIndex;
+					}
+				}
+
+				return DuplicatedVertCount;
+			}
+
+			/**
+			 * Make sure the SourceMeshDescription is compact before calling this function, we assume all element ID are continuous so we can use them as array index
+			 */
+			bool CopyMeshDescriptionToSkeletalMeshImportData(const FMeshDescription& SourceMeshDescription, FSkeletalMeshImportData& DestinationSkeletalMeshImportData)
+			{
+				DestinationSkeletalMeshImportData.Empty();
+
+				FSkeletalMeshConstAttributes Attributes(SourceMeshDescription);
+				TVertexAttributesConstRef<FVector> VertexPositions = Attributes.GetVertexPositions();
+				TVertexAttributesConstRef<TArrayAttribute<int32>> VertexInfluenceBones = Attributes.GetVertexInfluenceBones();
+				TVertexAttributesConstRef<TArrayAttribute<float>> VertexInfluenceWeights = Attributes.GetVertexInfluenceWeights();
+
+				TVertexInstanceAttributesConstRef<FVertexID> VertexInstanceVertexIndices = Attributes.GetVertexInstanceVertexIndices();
+				TVertexInstanceAttributesConstRef<FVector2D> VertexInstanceUVs = Attributes.GetVertexInstanceUVs();
+				TVertexInstanceAttributesConstRef<FVector> VertexInstanceNormals = Attributes.GetVertexInstanceNormals();
+				TVertexInstanceAttributesConstRef<FVector> VertexInstanceTangents = Attributes.GetVertexInstanceTangents();
+				TVertexInstanceAttributesConstRef<float> VertexInstanceBiNormalSigns = Attributes.GetVertexInstanceBinormalSigns();
+				TVertexInstanceAttributesConstRef<FVector4> VertexInstanceColors = Attributes.GetVertexInstanceColors();
+
+				TPolygonGroupAttributesConstRef<FName> PolygonGroupMaterialSlotNames = Attributes.GetPolygonGroupMaterialSlotNames();
+				//Get the per face smoothing
+				TArray<uint32> FaceSmoothingMasks;
+				FaceSmoothingMasks.AddZeroed(SourceMeshDescription.Triangles().Num());
+				FSkeletalMeshOperations::ConvertHardEdgesToSmoothGroup(SourceMeshDescription, FaceSmoothingMasks);
+
+				//////////////////////////////////////////////////////////////////////////
+				// Copy the materials
+				DestinationSkeletalMeshImportData.Materials.Reserve(SourceMeshDescription.PolygonGroups().Num());
+				for (FPolygonGroupID PolygonGroupID : SourceMeshDescription.PolygonGroups().GetElementIDs())
+				{
+					SkeletalMeshImportData::FMaterial Material;
+					Material.MaterialImportName = PolygonGroupMaterialSlotNames[PolygonGroupID].ToString();
+					//The material interface will be added later by the factory
+					DestinationSkeletalMeshImportData.Materials.Add(Material);
+				}
+				DestinationSkeletalMeshImportData.MaxMaterialIndex = DestinationSkeletalMeshImportData.Materials.Num()-1;
+
+				//////////////////////////////////////////////////////////////////////////
+				//Copy the vertex positions and the influences
+
+				//Reserve the point and influences
+				DestinationSkeletalMeshImportData.Points.AddZeroed(SourceMeshDescription.Vertices().Num());
+				DestinationSkeletalMeshImportData.Influences.Reserve(SourceMeshDescription.Vertices().Num() * 4);
+				
+				for (FVertexID VertexID : SourceMeshDescription.Vertices().GetElementIDs())
+				{
+					//We can use GetValue because the Meshdescription was compacted before the copy
+					DestinationSkeletalMeshImportData.Points[VertexID.GetValue()] = VertexPositions[VertexID];
+					int32 InfluenceCount = VertexInfluenceBones[VertexID].Num();
+					int32 InfluenceOffsetIndex = DestinationSkeletalMeshImportData.Influences.Num();
+					DestinationSkeletalMeshImportData.Influences.AddDefaulted(InfluenceCount);
+					for (int32 InfluenceIndex = 0; InfluenceIndex < InfluenceCount; ++InfluenceIndex)
+					{
+						SkeletalMeshImportData::FRawBoneInfluence& BoneInfluence = DestinationSkeletalMeshImportData.Influences[InfluenceOffsetIndex + InfluenceIndex];
+						BoneInfluence.VertexIndex = VertexID.GetValue();
+						BoneInfluence.BoneIndex = VertexInfluenceBones[VertexID][InfluenceIndex];
+						BoneInfluence.Weight = VertexInfluenceWeights[VertexID][InfluenceIndex];
+					}
+				}
+
+				//////////////////////////////////////////////////////////////////////////
+				//Copy the triangle and vertex instances
+				DestinationSkeletalMeshImportData.Faces.AddZeroed(SourceMeshDescription.Triangles().Num());
+				DestinationSkeletalMeshImportData.Wedges.Reserve(SourceMeshDescription.VertexInstances().Num());
+				DestinationSkeletalMeshImportData.NumTexCoords = FMath::Min<int32>(VertexInstanceUVs.GetNumChannels(), (int32)MAX_TEXCOORDS);
+				for (FTriangleID TriangleID : SourceMeshDescription.Triangles().GetElementIDs())
+				{
+					FPolygonGroupID PolygonGroupID = SourceMeshDescription.GetTrianglePolygonGroup(TriangleID);
+					TArrayView<const FVertexInstanceID> VertexInstances = SourceMeshDescription.GetTriangleVertexInstances(TriangleID);
+					int32 FaceIndex = TriangleID.GetValue();
+					if (!ensure(DestinationSkeletalMeshImportData.Faces.IsValidIndex(FaceIndex)))
+					{
+						//TODO log an error for the user
+						break;
+					}
+					SkeletalMeshImportData::FTriangle& Face = DestinationSkeletalMeshImportData.Faces[FaceIndex];
+					Face.MatIndex = PolygonGroupID.GetValue();
+					Face.SmoothingGroups = 0;
+					if (FaceSmoothingMasks.IsValidIndex(FaceIndex))
+					{
+						Face.SmoothingGroups = FaceSmoothingMasks[FaceIndex];
+					}
+					//Create the wedges
+					for (int32 Corner = 0; Corner < 3; ++Corner)
+					{
+						FVertexInstanceID VertexInstanceID = VertexInstances[Corner];
+						SkeletalMeshImportData::FVertex Wedge;
+						Wedge.VertexIndex = (uint32)SourceMeshDescription.GetVertexInstanceVertex(VertexInstances[Corner]).GetValue();
+						Wedge.MatIndex = Face.MatIndex;
+						const bool bSRGB = false; //avoid linear to srgb conversion
+						Wedge.Color = FLinearColor(VertexInstanceColors[VertexInstanceID]).ToFColor(bSRGB);
+						for (int32 UVChannelIndex = 0; UVChannelIndex < (int32)(DestinationSkeletalMeshImportData.NumTexCoords); ++UVChannelIndex)
+						{
+							Wedge.UVs[UVChannelIndex] = VertexInstanceUVs.Get(VertexInstanceID, UVChannelIndex);
+						}
+						Face.TangentX[Corner] = VertexInstanceTangents[VertexInstanceID];
+						Face.TangentZ[Corner] = VertexInstanceNormals[VertexInstanceID];
+						Face.TangentY[Corner] = FVector::CrossProduct(VertexInstanceNormals[VertexInstanceID], VertexInstanceTangents[VertexInstanceID]).GetSafeNormal() * VertexInstanceBiNormalSigns[VertexInstanceID];
+
+						Face.WedgeIndex[Corner] = DestinationSkeletalMeshImportData.Wedges.Add(Wedge);
+					}
+				}
+
+				CleanUpUnusedMaterials(DestinationSkeletalMeshImportData);
+				// reorder material according to "SKinXX" in material name
+				SetMaterialSkinXXOrder(DestinationSkeletalMeshImportData);
+				DoUnSmoothVerts(DestinationSkeletalMeshImportData);
+
+				return true;
+			}
+
+			/**
+			 * Fill the Materials array using the raw skeletalmesh geometry data (using material imported name)
+			 * Find the material from the dependencies of the skeletalmesh before searching in all package.
+			 */
+			void ProcessImportMeshMaterials(TArray<FSkeletalMaterial>& Materials, FSkeletalMeshImportData& ImportData, TMap<FName, UMaterialInterface*>& AvailableMaterials)
+			{
+				TArray <SkeletalMeshImportData::FMaterial>& ImportedMaterials = ImportData.Materials;
+				// If direct linkup of materials is requested, try to find them here - to get a texture name from a 
+				// material name, cut off anything in front of the dot (beyond are special flags).
+				int32 SkinOffset = INDEX_NONE;
+				for (int32 MatIndex = 0; MatIndex < ImportedMaterials.Num(); ++MatIndex)
+				{
+					const SkeletalMeshImportData::FMaterial& ImportedMaterial = ImportedMaterials[MatIndex];
+
+					UMaterialInterface* Material = nullptr;
+					//Remove _SkinXX from the material name
+					FString MaterialNameNoSkin = UE::Interchange::Material::RemoveSkinFromName(ImportedMaterial.MaterialImportName);
+					const FName SearchMaterialSlotName(*ImportedMaterial.MaterialImportName);
+					int32 MaterialIndex = 0;
+					FSkeletalMaterial* SkeletalMeshMaterialFind = Materials.FindByPredicate([&SearchMaterialSlotName, &MaterialIndex](const FSkeletalMaterial& ItemMaterial)
+					{
+						//Imported material slot name is available only WITH_EDITOR
+						FName ImportedMaterialSlot = NAME_None;
+#if WITH_EDITOR
+						ImportedMaterialSlot = ItemMaterial.ImportedMaterialSlotName;
+#else
+						ImportedMaterialSlot = ItemMaterial.MaterialSlotName;
+#endif
+						if (ImportedMaterialSlot != SearchMaterialSlotName)
+						{
+							MaterialIndex++;
+							return false;
+						}
+						return true;
+					});
+
+					if (SkeletalMeshMaterialFind != nullptr)
+					{
+						Material = SkeletalMeshMaterialFind->MaterialInterface;
+					}
+
+					if(!Material)
+					{
+						//Try to find the material in the skeletal mesh node dependencies (Materials are import before skeletal mesh when there is a dependency)
+						const FName MaterialFName = FName(*ImportedMaterial.MaterialImportName);
+						if (AvailableMaterials.Contains(FName(MaterialFName)))
+						{
+							Material = AvailableMaterials.FindChecked(MaterialFName);
+						}
+						else
+						{
+							//We did not found any material in the dependencies so try to find material everywhere
+							FString MaterialName = ImportedMaterial.MaterialImportName;
+							Material = FindObject<UMaterialInterface>(ANY_PACKAGE, *MaterialName);
+							if (Material == nullptr && MaterialNameNoSkin.Len() < MaterialName.Len())
+							{
+								//If we have remove the _skinXX search for material without the suffixe
+								Material = FindObject<UMaterialInterface>(ANY_PACKAGE, *MaterialNameNoSkin);
+							}
+						}
+					}
+					
+					const bool bEnableShadowCasting = true;
+					const bool bInRecomputeTangent = false;
+					Materials.Add(FSkeletalMaterial(Material, bEnableShadowCasting, bInRecomputeTangent, Material != nullptr ? Material->GetFName() : FName(*MaterialNameNoSkin), FName(*(ImportedMaterial.MaterialImportName))));
+				}
+
+				int32 NumMaterialsToAdd = FMath::Max<int32>(ImportedMaterials.Num(), ImportData.MaxMaterialIndex + 1);
+
+				// Pad the material pointers
+				while (NumMaterialsToAdd > Materials.Num())
+				{
+					UMaterialInterface* NullMaterialInterface = nullptr;
+					Materials.Add(FSkeletalMaterial(NullMaterialInterface));
+				}
+			}
+
+			void ProcessImportMeshInfluences(const int32 WedgeCount, TArray<SkeletalMeshImportData::FRawBoneInfluence>& Influences)
+			{
+
+				// Sort influences by vertex index.
+				struct FCompareVertexIndex
+				{
+					bool operator()(const SkeletalMeshImportData::FRawBoneInfluence& A, const SkeletalMeshImportData::FRawBoneInfluence& B) const
+					{
+						if (A.VertexIndex > B.VertexIndex) return false;
+						else if (A.VertexIndex < B.VertexIndex) return true;
+						else if (A.Weight < B.Weight) return false;
+						else if (A.Weight > B.Weight) return true;
+						else if (A.BoneIndex > B.BoneIndex) return false;
+						else if (A.BoneIndex < B.BoneIndex) return true;
+						else									  return  false;
+					}
+				};
+				Influences.Sort(FCompareVertexIndex());
+
+				TArray <SkeletalMeshImportData::FRawBoneInfluence> NewInfluences;
+				int32	LastNewInfluenceIndex = 0;
+				int32	LastVertexIndex = INDEX_NONE;
+				int32	InfluenceCount = 0;
+
+				float TotalWeight = 0.f;
+				const float MINWEIGHT = 0.01f;
+
+				int MaxVertexInfluence = 0;
+				float MaxIgnoredWeight = 0.0f;
+
+				//We have to normalize the data before filtering influences
+				//Because influence filtering is base on the normalize value.
+				//Some DCC like Daz studio don't have normalized weight
+				for (int32 i = 0; i < Influences.Num(); i++)
+				{
+					// if less than min weight, or it's more than 8, then we clear it to use weight
+					InfluenceCount++;
+					TotalWeight += Influences[i].Weight;
+					// we have all influence for the same vertex, normalize it now
+					if (i + 1 >= Influences.Num() || Influences[i].VertexIndex != Influences[i + 1].VertexIndex)
+					{
+						// Normalize the last set of influences.
+						if (InfluenceCount && (TotalWeight != 1.0f))
+						{
+							float OneOverTotalWeight = 1.f / TotalWeight;
+							for (int r = 0; r < InfluenceCount; r++)
+							{
+								Influences[i - r].Weight *= OneOverTotalWeight;
+							}
+						}
+
+						if (MaxVertexInfluence < InfluenceCount)
+						{
+							MaxVertexInfluence = InfluenceCount;
+						}
+
+						// clear to count next one
+						InfluenceCount = 0;
+						TotalWeight = 0.f;
+					}
+
+					if (InfluenceCount > MAX_TOTAL_INFLUENCES && Influences[i].Weight > MaxIgnoredWeight)
+					{
+						MaxIgnoredWeight = Influences[i].Weight;
+					}
+				}
+
+				// warn about too many influences
+				if (MaxVertexInfluence > MAX_TOTAL_INFLUENCES)
+				{
+					//TODO log a display message to the user
+					//UE_LOG(LogLODUtilities, Display, TEXT("Skeletal mesh (%s) influence count of %d exceeds max count of %d. Influence truncation will occur. Maximum Ignored Weight %f"), *MeshName, MaxVertexInfluence, MAX_TOTAL_INFLUENCES, MaxIgnoredWeight);
+				}
+
+				for (int32 i = 0; i < Influences.Num(); i++)
+				{
+					// we found next verts, normalize it now
+					if (LastVertexIndex != Influences[i].VertexIndex)
+					{
+						// Normalize the last set of influences.
+						if (InfluenceCount && (TotalWeight != 1.0f))
+						{
+							float OneOverTotalWeight = 1.f / TotalWeight;
+							for (int r = 0; r < InfluenceCount; r++)
+							{
+								NewInfluences[LastNewInfluenceIndex - r].Weight *= OneOverTotalWeight;
+							}
+						}
+
+						// now we insert missing verts
+						if (LastVertexIndex != INDEX_NONE)
+						{
+							int32 CurrentVertexIndex = Influences[i].VertexIndex;
+							for (int32 j = LastVertexIndex + 1; j < CurrentVertexIndex; j++)
+							{
+								// Add a 0-bone weight if none other present (known to happen with certain MAX skeletal setups).
+								LastNewInfluenceIndex = NewInfluences.AddUninitialized();
+								NewInfluences[LastNewInfluenceIndex].VertexIndex = j;
+								NewInfluences[LastNewInfluenceIndex].BoneIndex = 0;
+								NewInfluences[LastNewInfluenceIndex].Weight = 1.f;
+							}
+						}
+
+						// clear to count next one
+						InfluenceCount = 0;
+						TotalWeight = 0.f;
+						LastVertexIndex = Influences[i].VertexIndex;
+					}
+
+					// if less than min weight, or it's more than 8, then we clear it to use weight
+					if (Influences[i].Weight > MINWEIGHT && InfluenceCount < MAX_TOTAL_INFLUENCES)
+					{
+						LastNewInfluenceIndex = NewInfluences.Add(Influences[i]);
+						InfluenceCount++;
+						TotalWeight += Influences[i].Weight;
+					}
+				}
+
+				Influences = NewInfluences;
+
+				// Ensure that each vertex has at least one influence as e.g. CreateSkinningStream relies on it.
+				// The below code relies on influences being sorted by vertex index.
+				if (Influences.Num() == 0)
+				{
+					// warn about no influences
+					//TODO add a user log
+					//UE_LOG(LogLODUtilities, Warning, TEXT("Warning skeletal mesh (%s) has no vertex influences"), *MeshName);
+					// add one for each wedge entry
+					Influences.AddUninitialized(WedgeCount);
+					for (int32 WedgeIdx = 0; WedgeIdx < WedgeCount; WedgeIdx++)
+					{
+						Influences[WedgeIdx].VertexIndex = WedgeIdx;
+						Influences[WedgeIdx].BoneIndex = 0;
+						Influences[WedgeIdx].Weight = 1.0f;
+					}
+					for (int32 i = 0; i < Influences.Num(); i++)
+					{
+						int32 CurrentVertexIndex = Influences[i].VertexIndex;
+
+						if (LastVertexIndex != CurrentVertexIndex)
+						{
+							for (int32 j = LastVertexIndex + 1; j < CurrentVertexIndex; j++)
+							{
+								// Add a 0-bone weight if none other present (known to happen with certain MAX skeletal setups).
+								Influences.InsertUninitialized(i, 1);
+								Influences[i].VertexIndex = j;
+								Influences[i].BoneIndex = 0;
+								Influences[i].Weight = 1.f;
+							}
+							LastVertexIndex = CurrentVertexIndex;
+						}
+					}
+				}
+			}
+
+
+		} //Namespace Private
+	} //namespace Interchange
+} //namespace UE
+
+#endif //#if WITH_EDITOR
+
+
+UClass* UInterchangeSkeletalMeshFactory::GetFactoryClass() const
+{
+	return USkeletalMesh::StaticClass();
+}
+
+UObject* UInterchangeSkeletalMeshFactory::CreateEmptyAsset(const FCreateAssetParams& Arguments) const
+{
+#if !WITH_EDITOR || !WITH_EDITORONLY_DATA
+
+	UE_LOG(LogInterchangeImportPlugin, Error, TEXT("Cannot import skeletalMesh asset in runtime, this is an editor only feature."));
+	return nullptr;
+
+#else
+	USkeletalMesh* SkeletalMesh = nullptr;
+	if (!Arguments.AssetNode || !Arguments.AssetNode->GetAssetClass()->IsChildOf(GetFactoryClass()))
+	{
+		return nullptr;
+	}
+
+	const UInterchangeSkeletalMeshNode* SkeletalMeshNode = Cast<UInterchangeSkeletalMeshNode>(Arguments.AssetNode);
+	if (SkeletalMeshNode == nullptr)
+	{
+		return nullptr;
+	}
+
+	// create an asset if it doesn't exist
+	UObject* ExistingAsset = StaticFindObject(nullptr, Arguments.Parent, *Arguments.AssetName);
+
+	// create a new material or overwrite existing asset, if possible
+	if (!ExistingAsset)
+	{
+		SkeletalMesh = NewObject<USkeletalMesh>(Arguments.Parent, *Arguments.AssetName, RF_Public | RF_Standalone);
+	}
+	else if (ExistingAsset->GetClass()->IsChildOf(USkeletalMesh::StaticClass()))
+	{
+		//This is a reimport, we are just re-updating the source data
+		SkeletalMesh = Cast<USkeletalMesh>(ExistingAsset);
+	}
+	
+	if (!SkeletalMesh)
+	{
+		UE_LOG(LogInterchangeImportPlugin, Warning, TEXT("Could not create SkeletalMesh asset %s"), *Arguments.AssetName);
+		return nullptr;
+	}
+	
+	SkeletalMesh->PreEditChange(nullptr);
+	//Allocate the LODImport data in the main thread
+	SkeletalMesh->ReserveLODImportData(SkeletalMeshNode->GetLodDataCount());
+	
+	return SkeletalMesh;
+#endif //else !WITH_EDITOR || !WITH_EDITORONLY_DATA
+}
+
+UObject* UInterchangeSkeletalMeshFactory::CreateAsset(const UInterchangeSkeletalMeshFactory::FCreateAssetParams& Arguments) const
+{
+#if !WITH_EDITOR || !WITH_EDITORONLY_DATA
+
+	UE_LOG(LogInterchangeImportPlugin, Error, TEXT("Cannot import skeletalMesh asset in runtime, this is an editor only feature."));
+	return nullptr;
+
+#else
+
+	if (!Arguments.AssetNode || !Arguments.AssetNode->GetAssetClass()->IsChildOf(GetFactoryClass()))
+	{
+		return nullptr;
+	}
+
+	const UInterchangeSkeletalMeshNode* SkeletalMeshNode = Cast<UInterchangeSkeletalMeshNode>(Arguments.AssetNode);
+	if (SkeletalMeshNode == nullptr)
+	{
+		return nullptr;
+	}
+
+	const IInterchangeSkeletalMeshPayloadInterface* SkeletalMeshTranslatorPayloadInterface = Cast<IInterchangeSkeletalMeshPayloadInterface>(Arguments.Translator);
+	if (!SkeletalMeshTranslatorPayloadInterface)
+	{
+		UE_LOG(LogInterchangeImportPlugin, Error, TEXT("Cannot import skeletalMesh, the translator do not implement the IInterchangeSkeletalMeshPayloadInterface."));
+		return nullptr;
+	}
+
+	const UClass* SkeletalMeshClass = SkeletalMeshNode->GetAssetClass();
+	check(SkeletalMeshClass && SkeletalMeshClass->IsChildOf(GetFactoryClass()));
+
+	// create an asset if it doesn't exist
+	UObject* ExistingAsset = StaticFindObject(nullptr, Arguments.Parent, *Arguments.AssetName);
+
+	UObject* SkeletalMeshObject = nullptr;
+	// create a new material or overwrite existing asset, if possible
+	if (!ExistingAsset)
+	{
+		//NewObject is not thread safe, the asset registry directory watcher tick on the main thread can trig before we finish initializing the UObject and will crash
+		//The UObject should have been create by calling CreateEmptyAsset on the main thread.
+		check(IsInGameThread());
+		SkeletalMeshObject = NewObject<UObject>(Arguments.Parent, SkeletalMeshClass, *Arguments.AssetName, RF_Public | RF_Standalone);
+	}
+	else if(ExistingAsset->GetClass()->IsChildOf(SkeletalMeshClass))
+	{
+		//This is a reimport, we are just re-updating the source data
+		SkeletalMeshObject = ExistingAsset;
+	}
+
+	if (!SkeletalMeshObject)
+	{
+		UE_LOG(LogInterchangeImportPlugin, Error, TEXT("Could not create SkeletalMesh asset %s"), *Arguments.AssetName);
+		return nullptr;
+	}
+
+	if (SkeletalMeshObject)
+	{
+		//Currently material re-import will not touch the material at all
+		//TODO design a re-import process for the material (expressions and input connections)
+		if(!Arguments.ReimportObject)
+		{
+			USkeletalMesh* SkeletalMesh = Cast<USkeletalMesh>(SkeletalMeshObject);
+			if (!ensure(SkeletalMesh))
+			{
+				UE_LOG(LogInterchangeImportPlugin, Error, TEXT("Could not create SkeletalMesh asset %s"), *Arguments.AssetName);
+				return nullptr;
+			}
+			//Dirty the DDC Key for any imported Skeletal Mesh
+			SkeletalMesh->InvalidateDeriveDataCacheGUID();
+
+			FSkeletalMeshModel* ImportedResource = SkeletalMesh->GetImportedModel();
+			if (!ensure(ImportedResource->LODModels.Num() == 0))
+			{
+				ImportedResource->LODModels.Empty();
+			}
+			USkeleton* SkeletonReference = nullptr;
+			int32 LodCount = SkeletalMeshNode->GetLodDataCount();
+			TArray<FName> LodDataUniqueIds;
+			SkeletalMeshNode->GetLodDataUniqueIds(LodDataUniqueIds);
+			ensure(LodDataUniqueIds.Num() == LodCount);
+			int32 CurrentLodIndex = 0;
+			for (int32 LodIndex = 0; LodIndex < LodCount; ++LodIndex)
+			{
+				FName LodUniqueId = LodDataUniqueIds[LodIndex];
+				const UInterchangeSkeletalMeshLodDataNode* LodDataNode = Cast<UInterchangeSkeletalMeshLodDataNode>(Arguments.NodeContainer->GetNode(LodUniqueId));
+				if (!LodDataNode)
+				{
+					UE_LOG(LogInterchangeImportPlugin, Warning, TEXT("Invalid LOD when importing SkeletalMesh asset %s"), *Arguments.AssetName);
+					continue;
+				}
+				FName SkeletonNodeId = NAME_None;
+				if (!LodDataNode->GetCustomSkeletonID(SkeletonNodeId))
+				{
+					UE_LOG(LogInterchangeImportPlugin, Warning, TEXT("Invalid Skeleton LOD when importing SkeletalMesh asset %s"), *Arguments.AssetName);
+					continue;
+				}
+				const UInterchangeSkeletonNode* SkeletonNode = Cast<UInterchangeSkeletonNode>(Arguments.NodeContainer->GetNode(SkeletonNodeId));
+				if (!SkeletonNode)
+				{
+					UE_LOG(LogInterchangeImportPlugin, Warning, TEXT("Invalid Skeleton LOD when importing SkeletalMesh asset %s"), *Arguments.AssetName);
+					continue;
+				}
+				
+				if(SkeletonReference == nullptr)
+				{
+					UObject* SkeletonObject = SkeletonNode->ReferenceObject.ResolveObject();
+					if (SkeletonObject)
+					{
+						SkeletonReference = Cast<USkeleton>(SkeletonObject);
+					}
+					if (!ensure(SkeletonReference))
+					{
+						UE_LOG(LogInterchangeImportPlugin, Warning, TEXT("Invalid Skeleton LOD when importing SkeletalMesh asset %s"), *Arguments.AssetName);
+						break;
+					}
+				}
+
+				FName RootJointNodeId = NAME_None;
+				if (!SkeletonNode->GetCustomRootJointID(RootJointNodeId))
+				{
+					UE_LOG(LogInterchangeImportPlugin, Warning, TEXT("Invalid Skeleton LOD Root Joint when importing SkeletalMesh asset %s"), *Arguments.AssetName);
+					continue;
+				}
+				
+				int32 SkeletonDepth = 0;
+				TArray<SkeletalMeshImportData::FBone> RefBonesBinary;
+				UE::Interchange::Private::ProcessImportMeshSkeleton(SkeletonReference, SkeletalMesh->GetRefSkeleton(), SkeletonDepth, Arguments.NodeContainer, RootJointNodeId, RefBonesBinary);
+
+				//Add the lod mesh data to the skeletalmesh
+				FSkeletalMeshImportData SkeletalMeshImportData;
+				TArray<FName> TranslatorMeshKeys;
+				LodDataNode->GetTranslatorMeshKeys(TranslatorMeshKeys);
+				//We should have only one key per LOD
+				ensure(TranslatorMeshKeys.Num() == 1);
+				//Use the first geo LOD since we are suppose to have only one Key
+				if(TranslatorMeshKeys.IsValidIndex(0))
+				{
+					int32 MeshLodIndex = 0;
+					SkeletalMeshImportData.RefBonesBinary = RefBonesBinary;
+
+					TOptional<UE::Interchange::FSkeletalMeshLodPayloadData> LodMeshPayload = SkeletalMeshTranslatorPayloadInterface->GetSkeletalMeshLodPayloadData(TranslatorMeshKeys[MeshLodIndex].ToString());
+					if (!LodMeshPayload.IsSet())
+					{
+						UE_LOG(LogInterchangeImportPlugin, Warning, TEXT("Invalid Skeletal mesh payload key [%s] SkeletalMesh asset %s"), *TranslatorMeshKeys[MeshLodIndex].ToString(), *Arguments.AssetName);
+						break;
+					}
+					//TODO use the mesh description as the source import data for skeletalmesh
+					//Convert the meshdescription to skeletalmesh import data and set the data for the LOD, so the build will work
+					//Make sure all IDs are compact so we can use them directly when copying the data
+					FElementIDRemappings ElementIDRemappings;
+					LodMeshPayload->LodMeshDescription.Compact(ElementIDRemappings);
+					UE::Interchange::Private::CopyMeshDescriptionToSkeletalMeshImportData(LodMeshPayload->LodMeshDescription, SkeletalMeshImportData);
+				}
+
+				ensure(ImportedResource->LODModels.Add(new FSkeletalMeshLODModel()) == CurrentLodIndex);
+				FSkeletalMeshLODModel& LODModel = ImportedResource->LODModels[CurrentLodIndex];
+
+				TMap<FName, UMaterialInterface*> AvailableMaterials;
+				TArray<FName> Dependencies;
+				SkeletalMeshNode->GetDependecies(Dependencies);
+				for (int32 DependencyIndex = 0; DependencyIndex < Dependencies.Num(); ++DependencyIndex)
+				{
+					const UInterchangeMaterialNode* MaterialNode = Cast<UInterchangeMaterialNode>(Arguments.NodeContainer->GetNode(Dependencies[DependencyIndex]));
+					if (!MaterialNode || !MaterialNode->ReferenceObject.IsValid())
+					{
+						continue;
+					}
+					UMaterialInterface* MaterialInterface = Cast<UMaterialInterface>(MaterialNode->ReferenceObject.ResolveObject());
+					if (!MaterialInterface)
+					{
+						continue;
+					}
+					AvailableMaterials.Add(MaterialNode->GetDisplayLabel(), MaterialInterface);
+				}
+
+				UE::Interchange::Private::ProcessImportMeshMaterials(SkeletalMesh->GetMaterials(), SkeletalMeshImportData, AvailableMaterials);
+				UE::Interchange::Private::ProcessImportMeshInfluences(SkeletalMeshImportData.Wedges.Num(), SkeletalMeshImportData.Influences);
+				//Store the original fbx import data the SkelMeshImportDataPtr should not be modified after this
+				SkeletalMesh->SaveLODImportedData(CurrentLodIndex, SkeletalMeshImportData);
+				//We reimport both
+				SkeletalMesh->SetLODImportedDataVersions(CurrentLodIndex, ESkeletalMeshGeoImportVersions::LatestVersion, ESkeletalMeshSkinningImportVersions::LatestVersion);
+				FSkeletalMeshLODInfo& NewLODInfo = SkeletalMesh->AddLODInfo();
+				NewLODInfo.ReductionSettings.NumOfTrianglesPercentage = 1.0f;
+				NewLODInfo.ReductionSettings.NumOfVertPercentage = 1.0f;
+				NewLODInfo.ReductionSettings.MaxDeviationPercentage = 0.0f;
+				NewLODInfo.LODHysteresis = 0.02f;
+				NewLODInfo.bImportWithBaseMesh = true;
+
+				//Add the bound to the skeletal mesh
+				if (SkeletalMesh->GetImportedBounds().BoxExtent.IsNearlyZero())
+				{
+					FBox BoundingBox(SkeletalMeshImportData.Points.GetData(), SkeletalMeshImportData.Points.Num());
+					const FVector BoundingBoxSize = BoundingBox.GetSize();
+
+					if (SkeletalMeshImportData.Points.Num() > 2 && BoundingBoxSize.X < THRESH_POINTS_ARE_SAME && BoundingBoxSize.Y < THRESH_POINTS_ARE_SAME && BoundingBoxSize.Z < THRESH_POINTS_ARE_SAME)
+					{
+						//TODO log a user error
+						//AddTokenizedErrorMessage(FTokenizedMessage::Create(EMessageSeverity::Error, FText::Format(LOCTEXT("FbxSkeletaLMeshimport_ErrorMeshTooSmall", "Cannot import this mesh, the bounding box of this mesh is smaller than the supported threshold[{0}]."), FText::FromString(FString::Printf(TEXT("%f"), THRESH_POINTS_ARE_SAME)))), FFbxErrors::SkeletalMesh_FillImportDataFailed);
+					}
+					SkeletalMesh->SetImportedBounds(FBoxSphereBounds(BoundingBox));
+				}
+
+				CurrentLodIndex++;
+			}
+
+			SkeletonReference->MergeAllBonesToBoneTree(SkeletalMesh);
+			if (SkeletalMesh->GetSkeleton() != SkeletonReference)
+			{
+				SkeletalMesh->SetSkeleton(SkeletonReference);
+			}
+
+			SkeletalMesh->CalculateInvRefMatrices();
+
+			/** Apply all SkeletalMeshNode custom attributes to the material asset */
+			SkeletalMeshNode->ApplyAllCustomAttributeToAsset(SkeletalMesh);
+		}
+		
+		//Getting the file Hash will cache it into the source data
+		Arguments.SourceData->GetFileContentHash();
+
+		//The interchange completion task (call in the GameThread after the factories pass), will call PostEditChange which will trig another asynchronous system that will build all material in parallel
+	}
+	else
+	{
+		//The material is not a UMaterialInterface
+		SkeletalMeshObject->RemoveFromRoot();
+		SkeletalMeshObject->MarkPendingKill();
+	}
+	return SkeletalMeshObject;
+
+#endif //else !WITH_EDITOR || !WITH_EDITORONLY_DATA
+}
+
+/* This function is call in the completion task on the main thread, use it to call main thread post creation step for your assets*/
+void UInterchangeSkeletalMeshFactory::PostImportGameThreadCallback(const FPostImportGameThreadCallbackParams& Arguments) const
+{
+	check(IsInGameThread());
+	Super::PostImportGameThreadCallback(Arguments);
+
+	//TODO make sure this work at runtime
+#if WITH_EDITORONLY_DATA
+	if (ensure(Arguments.ImportedObject && Arguments.SourceData))
+	{
+		//We must call the Update of the asset source file in the main thread because UAssetImportData::Update execute some delegate we do not control
+		USkeletalMesh* SkeletalMesh = CastChecked<USkeletalMesh>(Arguments.ImportedObject);
+
+		UAssetImportData* ImportDataPtr = SkeletalMesh->GetAssetImportData();
+		UE::Interchange::FFactoryCommon::FUpdateImportAssetDataParameters UpdateImportAssetDataParameters(SkeletalMesh
+																										  , &ImportDataPtr
+																										  , Arguments.SourceData
+																										  , Arguments.NodeUniqueID
+																										  , Arguments.NodeContainer);
+		UE::Interchange::FFactoryCommon::UpdateImportAssetData(UpdateImportAssetDataParameters);
+	}
+#endif
+}
\ No newline at end of file
diff --git a/Engine/Plugins/Runtime/Interchange/Source/Import/Private/Mesh/InterchangeSkeletonFactory.cpp b/Engine/Plugins/Runtime/Interchange/Source/Import/Private/Mesh/InterchangeSkeletonFactory.cpp
new file mode 100644
--- /dev/null
+++ b/Engine/Plugins/Runtime/Interchange/Source/Import/Private/Mesh/InterchangeSkeletonFactory.cpp
@@ -0,0 +1,133 @@
+// Copyright Epic Games, Inc. All Rights Reserved. 
+#include "Mesh/InterchangeSkeletonFactory.h"
+
+#include "Animation/Skeleton.h"
+#include "InterchangeImportCommon.h"
+#include "InterchangeSkeletonNode.h"
+#include "InterchangeSourceData.h"
+#include "LogInterchangeImportPlugin.h"
+#include "Nodes/InterchangeBaseNode.h"
+#include "Nodes/InterchangeBaseNodeContainer.h"
+
+
+
+UClass* UInterchangeSkeletonFactory::GetFactoryClass() const
+{
+	return USkeleton::StaticClass();
+}
+
+UObject* UInterchangeSkeletonFactory::CreateEmptyAsset(const FCreateAssetParams& Arguments) const
+{
+	UObject* Skeleton = nullptr;
+
+#if WITH_EDITORONLY_DATA
+	if (!Arguments.AssetNode || !Arguments.AssetNode->GetAssetClass()->IsChildOf(GetFactoryClass()))
+	{
+		return nullptr;
+	}
+
+	const UInterchangeSkeletonNode* SkeletonNode = Cast<UInterchangeSkeletonNode>(Arguments.AssetNode);
+	if (SkeletonNode == nullptr)
+	{
+		return nullptr;
+	}
+
+	// create an asset if it doesn't exist
+	UObject* ExistingAsset = StaticFindObject(nullptr, Arguments.Parent, *Arguments.AssetName);
+
+	// create a new material or overwrite existing asset, if possible
+	if (!ExistingAsset)
+	{
+		Skeleton = NewObject<UObject>(Arguments.Parent, USkeleton::StaticClass(), *Arguments.AssetName, RF_Public | RF_Standalone);
+	}
+	else if (ExistingAsset->GetClass()->IsChildOf(USkeleton::StaticClass()))
+	{
+		//This is a reimport, we are just re-updating the source data
+		Skeleton = ExistingAsset;
+	}
+
+	if (!Skeleton)
+	{
+		UE_LOG(LogInterchangeImportPlugin, Warning, TEXT("Could not create Skeleton asset %s"), *Arguments.AssetName);
+		return nullptr;
+	}
+	Skeleton->PreEditChange(nullptr);
+#endif //WITH_EDITORONLY_DATA
+	return Skeleton;
+}
+
+UObject* UInterchangeSkeletonFactory::CreateAsset(const UInterchangeSkeletonFactory::FCreateAssetParams& Arguments) const
+{
+#if !WITH_EDITORONLY_DATA
+
+	UE_LOG(LogInterchangeImportPlugin, Error, TEXT("Cannot import Skeleton asset in runtime, this is an editor only feature."));
+	return nullptr;
+
+#else
+
+	if (!Arguments.AssetNode || !Arguments.AssetNode->GetAssetClass()->IsChildOf(GetFactoryClass()))
+	{
+		return nullptr;
+	}
+
+	const UInterchangeSkeletonNode* SkeletonNode = Cast<UInterchangeSkeletonNode>(Arguments.AssetNode);
+	if (SkeletonNode == nullptr)
+	{
+		return nullptr;
+	}
+
+	const UClass* SkeletonClass = SkeletonNode->GetAssetClass();
+	check(SkeletonClass && SkeletonClass->IsChildOf(GetFactoryClass()));
+
+	// create an asset if it doesn't exist
+	UObject* ExistingAsset = StaticFindObject(nullptr, Arguments.Parent, *Arguments.AssetName);
+
+	UObject* SkeletonObject = nullptr;
+	// create a new material or overwrite existing asset, if possible
+	if (!ExistingAsset)
+	{
+		//NewObject is not thread safe, the asset registry directory watcher tick on the main thread can trig before we finish initializing the UObject and will crash
+		//The UObject should have been create by calling CreateEmptyAsset on the main thread.
+		check(IsInGameThread());
+		SkeletonObject = NewObject<UObject>(Arguments.Parent, SkeletonClass, *Arguments.AssetName, RF_Public | RF_Standalone);
+	}
+	else if(ExistingAsset->GetClass()->IsChildOf(SkeletonClass))
+	{
+		//This is a reimport, we are just re-updating the source data
+		SkeletonObject = ExistingAsset;
+	}
+
+	if (!SkeletonObject)
+	{
+		UE_LOG(LogInterchangeImportPlugin, Warning, TEXT("Could not create Skeleton asset %s"), *Arguments.AssetName);
+		return nullptr;
+	}
+
+	if (SkeletonObject)
+	{
+		//Currently material re-import will not touch the material at all
+		//TODO design a re-import process for the material (expressions and input connections)
+		if(!Arguments.ReimportObject)
+		{
+			USkeleton* Skeleton = Cast<USkeleton>(SkeletonObject);
+			if (!ensure(Skeleton))
+			{
+				UE_LOG(LogInterchangeImportPlugin, Warning, TEXT("Could not create Skeleton asset %s"), *Arguments.AssetName);
+				return nullptr;
+			}
+		}
+		
+		//Getting the file Hash will cache it into the source data
+		Arguments.SourceData->GetFileContentHash();
+
+		//The interchange completion task (call in the GameThread after the factories pass), will call PostEditChange which will trig another asynchronous system that will build all material in parallel
+	}
+	else
+	{
+		//The material is not a UMaterialInterface
+		SkeletonObject->RemoveFromRoot();
+		SkeletonObject->MarkPendingKill();
+	}
+	return SkeletonObject;
+#endif
+}
diff --git a/Engine/Plugins/Runtime/Interchange/Source/Import/Private/Texture/InterchangeTextureFactory.cpp b/Engine/Plugins/Runtime/Interchange/Source/Import/Private/Texture/InterchangeTextureFactory.cpp
--- a/Engine/Plugins/Runtime/Interchange/Source/Import/Private/Texture/InterchangeTextureFactory.cpp
+++ b/Engine/Plugins/Runtime/Interchange/Source/Import/Private/Texture/InterchangeTextureFactory.cpp
@@ -74,6 +74,9 @@ UObject* UInterchangeTextureFactory::CreateEmptyAsset(const FCreateAssetParams&
 		UE_LOG(LogInterchangeImportPlugin, Warning, TEXT("Could not create texture asset %s"), *Arguments.AssetName);
 		return nullptr;
 	}
+
+	Texture->PreEditChange(nullptr);
+
 #endif //WITH_EDITORONLY_DATA
 	return Texture;
 }
diff --git a/Engine/Plugins/Runtime/Interchange/Source/Import/Public/Fbx/InterchangeFbxTranslator.h b/Engine/Plugins/Runtime/Interchange/Source/Import/Public/Fbx/InterchangeFbxTranslator.h
--- a/Engine/Plugins/Runtime/Interchange/Source/Import/Public/Fbx/InterchangeFbxTranslator.h
+++ b/Engine/Plugins/Runtime/Interchange/Source/Import/Public/Fbx/InterchangeFbxTranslator.h
@@ -80,7 +80,7 @@ public:
 	 * @param PayloadKey - The key to retrieve the a particular payload contain into the specified source data.
 	 * @return a PayloadData containing the imported data. The TOptional will not be set if there is an error.
 	 */
-	virtual TOptional<UE::Interchange::FStaticMeshPayloadData> GetStaticMeshPayloadData(const UInterchangeSourceData* SourceData, const FString& PayLoadKey) const override;
+	virtual TOptional<UE::Interchange::FStaticMeshPayloadData> GetStaticMeshPayloadData(const FString& PayLoadKey) const override;
 
 	/* IInterchangeStaticMeshPayloadInterface End */
 
@@ -96,13 +96,15 @@ public:
 	 * @param PayloadKey - The key to retrieve the a particular payload contain into the specified source data.
 	 * @return a PayloadData containing the imported data. The TOptional will not be set if there is an error.
 	 */
-	virtual TOptional<UE::Interchange::FSkeletalMeshPayloadData> GetSkeletalMeshPayloadData(const UInterchangeSourceData* SourceData, const FString& PayLoadKey) const override;
+	virtual TOptional<UE::Interchange::FSkeletalMeshLodPayloadData> GetSkeletalMeshLodPayloadData(const FString& PayLoadKey) const override;
 
 	/* IInterchangeSkeletalMeshPayloadInterface End */
 private:
 
 	FString CreateLoadFbxFileCommand(const FString& FbxFilePath) const;
 
+	FString CreateFetchPayloadFbxCommand(const FString& FbxPayloadKey) const;
+
 	//Dispatcher is mutable since it is create during the Translate operation
 	//We do not want to allocate the dispatcher and start the InterchangeWorker process
 	//in the constructor because Archetype, CDO and registered translators will
diff --git a/Engine/Plugins/Runtime/Interchange/Source/Import/Public/Mesh/InterchangeSkeletalMeshFactory.h b/Engine/Plugins/Runtime/Interchange/Source/Import/Public/Mesh/InterchangeSkeletalMeshFactory.h
new file mode 100644
--- /dev/null
+++ b/Engine/Plugins/Runtime/Interchange/Source/Import/Public/Mesh/InterchangeSkeletalMeshFactory.h
@@ -0,0 +1,30 @@
+// Copyright Epic Games, Inc. All Rights Reserved.
+
+#pragma once
+
+#include "CoreMinimal.h"
+#include "InterchangeFactoryBase.h"
+#include "UObject/Object.h"
+#include "UObject/ObjectMacros.h"
+
+#include "InterchangeSkeletalMeshFactory.generated.h"
+
+UCLASS(BlueprintType)
+class INTERCHANGEIMPORTPLUGIN_API UInterchangeSkeletalMeshFactory : public UInterchangeFactoryBase
+{
+	GENERATED_BODY()
+public:
+
+	//////////////////////////////////////////////////////////////////////////
+	// Interchange factory base interface begin
+
+	virtual UClass* GetFactoryClass() const override;
+	virtual UObject* CreateEmptyAsset(const FCreateAssetParams& Arguments) const override;
+	virtual UObject* CreateAsset(const FCreateAssetParams& Arguments) const override;
+	virtual void PostImportGameThreadCallback(const FPostImportGameThreadCallbackParams& Arguments) const override;
+
+	// Interchange factory base interface end
+	//////////////////////////////////////////////////////////////////////////
+};
+
+
diff --git a/Engine/Plugins/Runtime/Interchange/Source/Import/Public/Mesh/InterchangeSkeletalMeshPayload.h b/Engine/Plugins/Runtime/Interchange/Source/Import/Public/Mesh/InterchangeSkeletalMeshPayload.h
--- a/Engine/Plugins/Runtime/Interchange/Source/Import/Public/Mesh/InterchangeSkeletalMeshPayload.h
+++ b/Engine/Plugins/Runtime/Interchange/Source/Import/Public/Mesh/InterchangeSkeletalMeshPayload.h
@@ -3,18 +3,16 @@
 #pragma once
 
 #include "CoreMinimal.h"
-#include "Rendering/SkeletalMeshLODImporterData.h"
+#include "MeshDescription.h"
 
 namespace UE
 {
 	namespace Interchange
 	{
-		struct FSkeletalMeshPayloadData
+		struct FSkeletalMeshLodPayloadData
 		{
 			//Currently the skeletalmesh payload data is editor only, we have to move to something available at runtime
-#if WITH_EDITOR
-			TArray<FSkeletalMeshImportData> SkeletalMeshImportData;
-#endif
+			FMeshDescription LodMeshDescription;
 		};
 	}//ns Interchange
 }//ns UE
diff --git a/Engine/Plugins/Runtime/Interchange/Source/Import/Public/Mesh/InterchangeSkeletalMeshPayloadInterface.h b/Engine/Plugins/Runtime/Interchange/Source/Import/Public/Mesh/InterchangeSkeletalMeshPayloadInterface.h
--- a/Engine/Plugins/Runtime/Interchange/Source/Import/Public/Mesh/InterchangeSkeletalMeshPayloadInterface.h
+++ b/Engine/Plugins/Runtime/Interchange/Source/Import/Public/Mesh/InterchangeSkeletalMeshPayloadInterface.h
@@ -28,11 +28,10 @@ public:
 	 * Once the translation is done, the import process need a way to retrieve payload data.
 	 * This payload will be use by the factories to create the asset.
 	 *
-	 * @param SourceData - The source data containing the data to translate
 	 * @param PayloadKey - The key to retrieve the a particular payload contain into the specified source data.
-	 * @return a PayloadData containing the import image data. The TOptional will not be set if there is an error.
+	 * @return a PayloadData containing the data point by the payload key. The TOptional will not be set if there is an error.
 	 */
-	virtual TOptional<UE::Interchange::FSkeletalMeshPayloadData> GetSkeletalMeshPayloadData(const UInterchangeSourceData* SourceData, const FString& PayloadKey) const = 0;
+	virtual TOptional<UE::Interchange::FSkeletalMeshLodPayloadData> GetSkeletalMeshLodPayloadData(const FString& PayloadKey) const = 0;
 };
 
 
diff --git a/Engine/Plugins/Runtime/Interchange/Source/Import/Public/Mesh/InterchangeSkeletonFactory.h b/Engine/Plugins/Runtime/Interchange/Source/Import/Public/Mesh/InterchangeSkeletonFactory.h
new file mode 100644
--- /dev/null
+++ b/Engine/Plugins/Runtime/Interchange/Source/Import/Public/Mesh/InterchangeSkeletonFactory.h
@@ -0,0 +1,30 @@
+// Copyright Epic Games, Inc. All Rights Reserved.
+
+#pragma once
+
+#include "CoreMinimal.h"
+#include "InterchangeFactoryBase.h"
+#include "UObject/Object.h"
+#include "UObject/ObjectMacros.h"
+
+#include "InterchangeSkeletonFactory.generated.h"
+
+UCLASS(BlueprintType)
+class INTERCHANGEIMPORTPLUGIN_API UInterchangeSkeletonFactory : public UInterchangeFactoryBase
+{
+	GENERATED_BODY()
+public:
+
+	//////////////////////////////////////////////////////////////////////////
+	// Interchange factory base interface begin
+
+	virtual UClass* GetFactoryClass() const override;
+	virtual UObject* CreateEmptyAsset(const FCreateAssetParams& Arguments) const override;
+	virtual UObject* CreateAsset(const FCreateAssetParams& Arguments) const override;
+	//virtual void PostImportGameThreadCallback(const FPostImportGameThreadCallbackParams& Arguments) const override;
+
+	// Interchange factory base interface end
+	//////////////////////////////////////////////////////////////////////////
+};
+
+
diff --git a/Engine/Plugins/Runtime/Interchange/Source/Import/Public/Mesh/InterchangeStaticMeshPayloadInterface.h b/Engine/Plugins/Runtime/Interchange/Source/Import/Public/Mesh/InterchangeStaticMeshPayloadInterface.h
--- a/Engine/Plugins/Runtime/Interchange/Source/Import/Public/Mesh/InterchangeStaticMeshPayloadInterface.h
+++ b/Engine/Plugins/Runtime/Interchange/Source/Import/Public/Mesh/InterchangeStaticMeshPayloadInterface.h
@@ -28,11 +28,10 @@ public:
 	 * Once the translation is done, the import process need a way to retrieve payload data.
 	 * This payload will be use by the factories to create the asset.
 	 *
-	 * @param SourceData - The source data containing the data to translate
 	 * @param PayloadKey - The key to retrieve the a particular payload contain into the specified source data.
-	 * @return a PayloadData containing the import image data. The TOptional will not be set if there is an error.
+	 * @return a PayloadData containing the the data ask with the key.
 	 */
-	virtual TOptional<UE::Interchange::FStaticMeshPayloadData> GetStaticMeshPayloadData(const UInterchangeSourceData* SourceData, const FString& PayloadKey) const = 0;
+	virtual TOptional<UE::Interchange::FStaticMeshPayloadData> GetStaticMeshPayloadData(const FString& PayloadKey) const = 0;
 };
 
 
diff --git a/Engine/Plugins/Runtime/Interchange/Source/Import/Public/Texture/InterchangeTexturePayloadData.h b/Engine/Plugins/Runtime/Interchange/Source/Import/Public/Texture/InterchangeTexturePayloadData.h
--- a/Engine/Plugins/Runtime/Interchange/Source/Import/Public/Texture/InterchangeTexturePayloadData.h
+++ b/Engine/Plugins/Runtime/Interchange/Source/Import/Public/Texture/InterchangeTexturePayloadData.h
@@ -9,7 +9,7 @@ namespace UE
 {
 	namespace Interchange
 	{
-		struct FImportImage
+		struct INTERCHANGEIMPORTPLUGIN_API FImportImage
 		{
 			TArray64<uint8> RawData;
 			ETextureSourceFormat Format = TSF_Invalid;
@@ -28,7 +28,7 @@ namespace UE
 			void* GetMipData(int32 InMipIndex);
 		};
 
-		struct FImportImageHelper
+		struct INTERCHANGEIMPORTPLUGIN_API FImportImageHelper
 		{
 			/**
 			 * Tests if the given height and width specify a supported texture resolution to import; Can optionally check if the height/width are powers of two
diff --git a/Engine/Plugins/Runtime/Interchange/Source/Nodes/Public/InterchangeJointNode.h b/Engine/Plugins/Runtime/Interchange/Source/Nodes/Public/InterchangeJointNode.h
new file mode 100644
--- /dev/null
+++ b/Engine/Plugins/Runtime/Interchange/Source/Nodes/Public/InterchangeJointNode.h
@@ -0,0 +1,80 @@
+// Copyright Epic Games, Inc. All Rights Reserved.
+
+#pragma once
+
+#include "CoreMinimal.h"
+#include "Nodes/InterchangeBaseNode.h"
+
+#include "InterchangeJointNode.generated.h"
+
+/**
+ * The joint node is use to store the joint information available when the translator read the data
+ */
+
+
+UCLASS(BlueprintType)
+class INTERCHANGENODEPLUGIN_API UInterchangeJointNode : public UInterchangeBaseNode
+{
+	GENERATED_BODY()
+
+public:
+	UInterchangeJointNode()
+	:UInterchangeBaseNode()
+	{
+	}
+
+	/**
+	 * Return the node type name of the class, we use this when reporting error
+	 */
+	virtual FString GetTypeName() const override
+	{
+		const FString TypeName = TEXT("JointNode");
+		return TypeName;
+	}
+
+	/** Get the class this node want to create */
+	UFUNCTION(BlueprintCallable, Category = "Interchange | Node | Joint")
+	virtual class UClass* GetAssetClass() const override
+	{
+		//No factory for joint so just return nullptr here, it will be import with the skeleton factory
+		return nullptr;
+	}
+
+	virtual FGuid GetHash() const override
+	{
+		return Attributes.GetStorageHash();
+	}
+
+public:
+	//Joint name attribute, all joints in a skeleton must have unique name
+
+	UFUNCTION(BlueprintCallable, Category = "Interchange | Node | Joint")
+	bool GetCustomName(FName& AttributeValue) const
+	{
+		IMPLEMENT_NODE_ATTRIBUTE_GETTER(Name, FName);
+	}
+
+	UFUNCTION(BlueprintCallable, Category = "Interchange | Node | Joint")
+	bool SetCustomName(const FName& AttributeValue)
+	{
+		IMPLEMENT_NODE_ATTRIBUTE_SETTER_NODELEGATE(Name, FName)
+	}
+
+	//Local transform attribute
+
+	UFUNCTION(BlueprintCallable, Category = "Interchange | Node | Joint")
+	bool GetCustomLocalTransform(FTransform& AttributeValue) const
+	{
+		IMPLEMENT_NODE_ATTRIBUTE_GETTER(LocalTransform, FTransform);
+	}
+
+	UFUNCTION(BlueprintCallable, Category = "Interchange | Node | Joint")
+	bool SetCustomLocalTransform(const FTransform& AttributeValue)
+	{
+		IMPLEMENT_NODE_ATTRIBUTE_SETTER_NODELEGATE(LocalTransform, FTransform)
+	}
+
+private:
+	const UE::Interchange::FAttributeKey Macro_CustomNameKey = UE::Interchange::FAttributeKey(TEXT("JointName"));
+	const UE::Interchange::FAttributeKey Macro_CustomLocalTransformKey = UE::Interchange::FAttributeKey(TEXT("LocalTransform"));
+};
diff --git a/Engine/Plugins/Runtime/Interchange/Source/Nodes/Public/InterchangeMaterialNode.h b/Engine/Plugins/Runtime/Interchange/Source/Nodes/Public/InterchangeMaterialNode.h
--- a/Engine/Plugins/Runtime/Interchange/Source/Nodes/Public/InterchangeMaterialNode.h
+++ b/Engine/Plugins/Runtime/Interchange/Source/Nodes/Public/InterchangeMaterialNode.h
@@ -20,6 +20,27 @@ namespace UE
 {
 	namespace Interchange
 	{
+		namespace Material
+		{
+			/**
+			 * Remove _SkinXX from a string
+			 */
+			static FString RemoveSkinFromName(const FString& Name)
+			{
+				int32 SkinOffset = Name.Find(TEXT("_skin"), ESearchCase::IgnoreCase, ESearchDir::FromEnd);
+				if (SkinOffset != INDEX_NONE)
+				{
+					FString MaterialNameNoSkin = Name;
+					FString SkinXXNumber = Name.Right(Name.Len() - (SkinOffset + 1)).RightChop(4);
+					if (SkinXXNumber.IsNumeric())
+					{
+						MaterialNameNoSkin = Name.LeftChop(Name.Len() - SkinOffset);
+					}
+					return MaterialNameNoSkin;
+				}
+				return Name;
+			}
+		} //ns Material
 
 		struct FMaterialNodeStaticData : public FBaseNodeStaticData
 		{
diff --git a/Engine/Plugins/Runtime/Interchange/Source/Nodes/Public/InterchangeSkeletalMeshLodDataNode.h b/Engine/Plugins/Runtime/Interchange/Source/Nodes/Public/InterchangeSkeletalMeshLodDataNode.h
new file mode 100644
--- /dev/null
+++ b/Engine/Plugins/Runtime/Interchange/Source/Nodes/Public/InterchangeSkeletalMeshLodDataNode.h
@@ -0,0 +1,223 @@
+// Copyright Epic Games, Inc. All Rights Reserved.
+
+#pragma once
+
+#include "CoreMinimal.h"
+#include "Nodes/InterchangeBaseNode.h"
+
+#include "InterchangeSkeletalMeshLodDataNode.generated.h"
+
+//Interchange namespace
+namespace UE
+{
+	namespace Interchange
+	{
+
+		struct FSkeletalMeshNodeLodDataStaticData : public FBaseNodeStaticData
+		{
+			static const FAttributeKey& TranslatorMeshKeysCountKey()
+			{
+				static FAttributeKey AttributeKey(TEXT("__LOD_DATA_IDS_COUNT_"));
+				return AttributeKey;
+			}
+
+			static const FAttributeKey& TranslatorMeshKeysBaseKey()
+			{
+				static FAttributeKey AttributeKey(TEXT("__LOD_DATA_IDS_INDEX_"));
+				return AttributeKey;
+			}
+		};
+
+	}//ns Interchange
+}//ns UE
+
+UCLASS(BlueprintType)
+class INTERCHANGENODEPLUGIN_API UInterchangeSkeletalMeshLodDataNode : public UInterchangeBaseNode
+{
+	GENERATED_BODY()
+
+public:
+	UInterchangeSkeletalMeshLodDataNode()
+		:UInterchangeBaseNode()
+	{
+	}
+
+	/**
+	 * Return the node type name of the class, we use this when reporting error
+	 */
+	virtual FString GetTypeName() const override
+	{
+		const FString TypeName = TEXT("SkeletalMeshLodDataNode");
+		return TypeName;
+	}
+
+	virtual FGuid GetHash() const override
+	{
+		return Attributes.GetStorageHash();
+	}
+
+
+public:
+	/** Return false if the Attribute was not set previously.*/
+	UFUNCTION(BlueprintCallable, Category = "Interchange | Node | SkeletalMeshLodData")
+	bool GetCustomSkeletonID(FName& AttributeValue) const
+	{
+		IMPLEMENT_NODE_ATTRIBUTE_GETTER(SkeletonID, FName);
+	}
+
+	UFUNCTION(BlueprintCallable, Category = "Interchange | Node | SkeletalMeshLodData")
+	bool SetCustomSkeletonID(const FName& AttributeValue)
+	{
+		IMPLEMENT_NODE_ATTRIBUTE_SETTER_NODELEGATE(SkeletonID, FName)
+	}
+
+	/** Return The number of LOD this skeletalmesh has.*/
+	UFUNCTION(BlueprintCallable, Category = "Interchange | Node | SkeletalMeshLodData")
+	int32 GetTranslatorMeshKeyCount() const
+	{
+		int32 MeshCount = 0;
+		if (Attributes.ContainAttribute(UE::Interchange::FSkeletalMeshNodeLodDataStaticData::TranslatorMeshKeysCountKey()))
+		{
+			UE::Interchange::FAttributeStorage::TAttributeHandle<int32> Handle = Attributes.GetAttributeHandle<int32>(UE::Interchange::FSkeletalMeshNodeLodDataStaticData::TranslatorMeshKeysCountKey());
+			if (Handle.IsValid())
+			{
+				Handle.Get(MeshCount);
+			}
+		}
+		return MeshCount;
+	}
+
+	UFUNCTION(BlueprintCallable, Category = "Interchange | Node | SkeletalMeshLodData")
+	void GetTranslatorMeshKeys(TArray<FName>& OutTranslatorMeshKeys) const
+	{
+		OutTranslatorMeshKeys.Reset();
+
+		int32 MeshCount = 0;
+		if (!Attributes.ContainAttribute(UE::Interchange::FSkeletalMeshNodeLodDataStaticData::TranslatorMeshKeysCountKey()))
+		{
+			return;
+		}
+		UE::Interchange::FAttributeStorage::TAttributeHandle<int32> Handle = Attributes.GetAttributeHandle<int32>(UE::Interchange::FSkeletalMeshNodeLodDataStaticData::TranslatorMeshKeysCountKey());
+		if (!Handle.IsValid())
+		{
+			return;
+		}
+		Handle.Get(MeshCount);
+		for (int32 MeshIndex = 0; MeshIndex < MeshCount; ++MeshIndex)
+		{
+			FString DepIndexKeyString = UE::Interchange::FSkeletalMeshNodeLodDataStaticData::TranslatorMeshKeysBaseKey().ToString() + FString::FromInt(MeshIndex);
+			UE::Interchange::FAttributeKey DepIndexKey(*DepIndexKeyString);
+			UE::Interchange::FAttributeStorage::TAttributeHandle<FName> HandleMesh = Attributes.GetAttributeHandle<FName>(DepIndexKey);
+			if (!HandleMesh.IsValid())
+			{
+				continue;
+			}
+			FName& NodeUniqueID = OutTranslatorMeshKeys.AddDefaulted_GetRef();
+			HandleMesh.Get(NodeUniqueID);
+		}
+	}
+
+	UFUNCTION(BlueprintCallable, Category = "Interchange | Node | SkeletalMeshLodData")
+	bool AddTranslatorMeshKey(FName TranslatorMeshKey)
+	{
+		if (!Attributes.ContainAttribute(UE::Interchange::FSkeletalMeshNodeLodDataStaticData::TranslatorMeshKeysCountKey()))
+		{
+			const int32 DependencyCount = 0;
+			UE::Interchange::EAttributeStorageResult Result = Attributes.RegisterAttribute<int32>(UE::Interchange::FSkeletalMeshNodeLodDataStaticData::TranslatorMeshKeysCountKey(), DependencyCount);
+			if (!UE::Interchange::IsAttributeStorageResultSuccess(Result))
+			{
+				UE::Interchange::LogAttributeStorageErrors(Result, TEXT("UInterchangeSkeletalMeshLodDataNode.AddTranslatorMeshKey"), UE::Interchange::FSkeletalMeshNodeLodDataStaticData::TranslatorMeshKeysCountKey());
+				return false;
+			}
+		}
+		UE::Interchange::FAttributeStorage::TAttributeHandle<int32> Handle = Attributes.GetAttributeHandle<int32>(UE::Interchange::FSkeletalMeshNodeLodDataStaticData::TranslatorMeshKeysCountKey());
+		if (!ensure(Handle.IsValid()))
+		{
+			return false;
+		}
+		int32 MeshIndex = 0;
+		Handle.Get(MeshIndex);
+		FString MeshIndexKeyString = UE::Interchange::FSkeletalMeshNodeLodDataStaticData::TranslatorMeshKeysBaseKey().ToString() + FString::FromInt(MeshIndex);
+		UE::Interchange::FAttributeKey MeshIndexKey(*MeshIndexKeyString);
+		//Increment the counter
+		MeshIndex++;
+		Handle.Set(MeshIndex);
+
+		UE::Interchange::EAttributeStorageResult AddMeshResult = Attributes.RegisterAttribute<FName>(MeshIndexKey, TranslatorMeshKey);
+		return true;
+	}
+
+	UFUNCTION(BlueprintCallable, Category = "Interchange | Node | SkeletalMeshLodData")
+	bool RemoveTranslatorMeshKey(FName TranslatorMeshKey)
+	{
+		int32 MeshCount = 0;
+		if (!Attributes.ContainAttribute(UE::Interchange::FSkeletalMeshNodeLodDataStaticData::TranslatorMeshKeysCountKey()))
+		{
+			return false;
+		}
+		UE::Interchange::FAttributeStorage::TAttributeHandle<int32> Handle = Attributes.GetAttributeHandle<int32>(UE::Interchange::FSkeletalMeshNodeLodDataStaticData::TranslatorMeshKeysCountKey());
+		if (!Handle.IsValid())
+		{
+			return false;
+		}
+		Handle.Get(MeshCount);
+		bool DecrementKey = false;
+		for (int32 MeshIndex = 0; MeshIndex < MeshCount; ++MeshIndex)
+		{
+			FString DepIndexKeyString = UE::Interchange::FSkeletalMeshNodeLodDataStaticData::TranslatorMeshKeysBaseKey().ToString() + FString::FromInt(MeshIndex);
+			UE::Interchange::FAttributeKey DepIndexKey(*DepIndexKeyString);
+			UE::Interchange::FAttributeStorage::TAttributeHandle<FName> HandleMesh = Attributes.GetAttributeHandle<FName>(DepIndexKey);
+			if (!HandleMesh.IsValid())
+			{
+				continue;
+			}
+			FName MeshUniqueID;
+			HandleMesh.Get(MeshUniqueID);
+			if (MeshUniqueID == TranslatorMeshKey)
+			{
+				//Remove this entry
+				Attributes.UnregisterAttribute(DepIndexKey);
+				Handle.Set(MeshCount-1);
+				//We have to rename the key for all the next item
+				DecrementKey = true;;
+			}
+			else if (DecrementKey)
+			{
+				FString NewDepIndexKeyString = UE::Interchange::FSkeletalMeshNodeLodDataStaticData::TranslatorMeshKeysBaseKey().ToString() + FString::FromInt(MeshIndex-1);
+				UE::Interchange::FAttributeKey NewDepIndexKey(*NewDepIndexKeyString);
+				UE::Interchange::EAttributeStorageResult UnregisterResult = Attributes.UnregisterAttribute(DepIndexKey);
+				if (UE::Interchange::IsAttributeStorageResultSuccess(UnregisterResult))
+				{
+					UE::Interchange::EAttributeStorageResult RegisterResult = Attributes.RegisterAttribute<FName>(NewDepIndexKey, MeshUniqueID);
+					if (!UE::Interchange::IsAttributeStorageResultSuccess(RegisterResult))
+					{
+						UE::Interchange::LogAttributeStorageErrors(RegisterResult, TEXT("UInterchangeSkeletalMeshLodDataNode.RemoveTranslatorMeshKey"), NewDepIndexKey);
+					}
+				}
+				else
+				{
+					UE::Interchange::LogAttributeStorageErrors(UnregisterResult, TEXT("UInterchangeSkeletalMeshLodDataNode.RemoveTranslatorMeshKey"), DepIndexKey);
+				}
+				//make sure the for loop iterate since HandleMesh is now invalid
+				continue;
+			}
+		}
+		return true;
+	}
+
+private:
+
+	bool IsEditorOnlyDataDefined()
+	{
+#if WITH_EDITORONLY_DATA
+		return true;
+#else
+		return false;
+#endif
+	}
+
+	//SkeletalMesh
+	const UE::Interchange::FAttributeKey Macro_CustomSkeletonIDKey = UE::Interchange::FAttributeKey(TEXT("SkeletonID"));
+
+protected:
+};
diff --git a/Engine/Plugins/Runtime/Interchange/Source/Nodes/Public/InterchangeSkeletalMeshNode.h b/Engine/Plugins/Runtime/Interchange/Source/Nodes/Public/InterchangeSkeletalMeshNode.h
new file mode 100644
--- /dev/null
+++ b/Engine/Plugins/Runtime/Interchange/Source/Nodes/Public/InterchangeSkeletalMeshNode.h
@@ -0,0 +1,288 @@
+// Copyright Epic Games, Inc. All Rights Reserved.
+
+#pragma once
+
+#include "CoreMinimal.h"
+#include "Nodes/InterchangeBaseNode.h"
+
+#if WITH_ENGINE
+#include "Engine/SkeletalMesh.h"
+#endif
+
+#include "InterchangeSkeletalMeshNode.generated.h"
+
+//Interchange namespace
+namespace UE
+{
+	namespace Interchange
+	{
+
+		struct FSkeletalMeshNodeStaticData : public FBaseNodeStaticData
+		{
+			static const FAttributeKey& LodDataIdsCountKey()
+			{
+				static FAttributeKey AttributeKey(TEXT("__LOD_DATA_IDS_COUNT_"));
+				return AttributeKey;
+			}
+
+			static const FAttributeKey& LodDataIdsBaseKey()
+			{
+				static FAttributeKey AttributeKey(TEXT("__LOD_DATA_IDS_INDEX_"));
+				return AttributeKey;
+			}
+		};
+
+	}//ns Interchange
+}//ns UE
+
+UCLASS(BlueprintType)
+class INTERCHANGENODEPLUGIN_API UInterchangeSkeletalMeshNode : public UInterchangeBaseNode
+{
+	GENERATED_BODY()
+
+public:
+	UInterchangeSkeletalMeshNode()
+		:UInterchangeBaseNode()
+	{
+#if WITH_ENGINE
+		AssetClass = nullptr;
+#endif
+	}
+
+	/**
+	 * Initialize node data
+	 * @param: UniqueID - The uniqueId for this node
+	 * @param DisplayLabel - The name of the node
+	 * @param InAssetClass - The class the SkeletalMesh factory will create for this node.
+	 */
+	UFUNCTION(BlueprintCallable, Category = "Interchange | Node | SkeletalMesh")
+	void InitializeSkeletalMeshNode(const FName& UniqueID, const FName& DisplayLabel, const FString& InAssetClass)
+	{
+		bIsNodeClassInitialized = false;
+		InitializeNode(UniqueID, DisplayLabel);
+		FString OperationName = GetTypeName() + TEXT(".SetAssetClassName");
+		InterchangePrivateNodeBase::SetCustomAttribute<FString>(Attributes, ClassNameAttributeKey, OperationName, InAssetClass);
+		FillAssetClassFromAttribute();
+	}
+
+	virtual void Serialize(FArchive& Ar) override
+	{
+		Super::Serialize(Ar);
+#if WITH_ENGINE
+		if (Ar.IsLoading())
+		{
+			//Make sure the class is properly set when we compile with engine, this will set the
+			//bIsNodeClassInitialized to true.
+			SetNodeClassFromClassAttribute();
+		}
+#endif //#if WITH_ENGINE
+	}
+
+	/**
+	 * Return the node type name of the class, we use this when reporting error
+	 */
+	virtual FString GetTypeName() const override
+	{
+		const FString TypeName = TEXT("SkeletalMeshNode");
+		return TypeName;
+	}
+
+	/** Get the class this node want to create */
+	UFUNCTION(BlueprintCallable, Category = "Interchange | Node | SkeletalMesh")
+	virtual class UClass* GetAssetClass() const override
+	{
+		ensure(bIsNodeClassInitialized);
+#if WITH_ENGINE
+		return AssetClass.Get() != nullptr ? AssetClass.Get() : USkeletalMesh::StaticClass();
+#else
+		return nullptr;
+#endif
+	}
+
+	virtual FGuid GetHash() const override
+	{
+		return Attributes.GetStorageHash();
+	}
+
+public:
+	/** Return The number of LOD this skeletalmesh has.*/
+	UFUNCTION(BlueprintCallable, Category = "Interchange | Node | SkeletalMesh")
+	int32 GetLodDataCount() const
+	{
+		int32 LodCount = 0;
+		if (Attributes.ContainAttribute(UE::Interchange::FSkeletalMeshNodeStaticData::LodDataIdsCountKey()))
+		{
+			UE::Interchange::FAttributeStorage::TAttributeHandle<int32> Handle = Attributes.GetAttributeHandle<int32>(UE::Interchange::FSkeletalMeshNodeStaticData::LodDataIdsCountKey());
+			if (Handle.IsValid())
+			{
+				Handle.Get(LodCount);
+			}
+		}
+		return LodCount;
+	}
+
+	UFUNCTION(BlueprintCallable, Category = "Interchange | Node | SkeletalMesh")
+	void GetLodDataUniqueIds(TArray<FName>& OutLodDataUniqueIds) const
+	{
+		OutLodDataUniqueIds.Reset();
+
+		int32 LodCount = 0;
+		if (!Attributes.ContainAttribute(UE::Interchange::FSkeletalMeshNodeStaticData::LodDataIdsCountKey()))
+		{
+			return;
+		}
+		UE::Interchange::FAttributeStorage::TAttributeHandle<int32> Handle = Attributes.GetAttributeHandle<int32>(UE::Interchange::FSkeletalMeshNodeStaticData::LodDataIdsCountKey());
+		if (!Handle.IsValid())
+		{
+			return;
+		}
+		Handle.Get(LodCount);
+		for (int32 LODIndex = 0; LODIndex < LodCount; ++LODIndex)
+		{
+			FString DepIndexKeyString = UE::Interchange::FSkeletalMeshNodeStaticData::LodDataIdsBaseKey().ToString() + FString::FromInt(LODIndex);
+			UE::Interchange::FAttributeKey DepIndexKey(*DepIndexKeyString);
+			UE::Interchange::FAttributeStorage::TAttributeHandle<FName> HandleLOD = Attributes.GetAttributeHandle<FName>(DepIndexKey);
+			if (!HandleLOD.IsValid())
+			{
+				continue;
+			}
+			FName& NodeUniqueID = OutLodDataUniqueIds.AddDefaulted_GetRef();
+			HandleLOD.Get(NodeUniqueID);
+		}
+	}
+
+	UFUNCTION(BlueprintCallable, Category = "Interchange | Node | SkeletalMesh")
+	bool AddLodDataUniqueId(FName LodDataUniqueId)
+	{
+		if (!Attributes.ContainAttribute(UE::Interchange::FSkeletalMeshNodeStaticData::LodDataIdsCountKey()))
+		{
+			const int32 DependencyCount = 0;
+			UE::Interchange::EAttributeStorageResult Result = Attributes.RegisterAttribute<int32>(UE::Interchange::FSkeletalMeshNodeStaticData::LodDataIdsCountKey(), DependencyCount);
+			if (!UE::Interchange::IsAttributeStorageResultSuccess(Result))
+			{
+				UE::Interchange::LogAttributeStorageErrors(Result, TEXT("UInterchangeSkeletalMeshNode.AddLodDataUniqueId"), UE::Interchange::FSkeletalMeshNodeStaticData::LodDataIdsCountKey());
+				return false;
+			}
+		}
+		UE::Interchange::FAttributeStorage::TAttributeHandle<int32> Handle = Attributes.GetAttributeHandle<int32>(UE::Interchange::FSkeletalMeshNodeStaticData::LodDataIdsCountKey());
+		if (!ensure(Handle.IsValid()))
+		{
+			return false;
+		}
+		int32 LODIndex = 0;
+		Handle.Get(LODIndex);
+		FString LodIndexKeyString = UE::Interchange::FSkeletalMeshNodeStaticData::LodDataIdsBaseKey().ToString() + FString::FromInt(LODIndex);
+		UE::Interchange::FAttributeKey LodIndexKey(*LodIndexKeyString);
+		//Increment the counter
+		LODIndex++;
+		Handle.Set(LODIndex);
+
+		UE::Interchange::EAttributeStorageResult AddLodResult = Attributes.RegisterAttribute<FName>(LodIndexKey, LodDataUniqueId);
+		if (!UE::Interchange::IsAttributeStorageResultSuccess(AddLodResult))
+		{
+			UE::Interchange::LogAttributeStorageErrors(AddLodResult, TEXT("UInterchangeSkeletalMeshNode.AddLodDataUniqueId"), LodIndexKey);
+			return false;
+		}
+		return true;
+	}
+
+	UFUNCTION(BlueprintCallable, Category = "Interchange | Node | SkeletalMesh")
+	bool RemoveLodDataUniqueId(FName LodDataUniqueId)
+	{
+		int32 LodCount = 0;
+		if (!Attributes.ContainAttribute(UE::Interchange::FSkeletalMeshNodeStaticData::LodDataIdsCountKey()))
+		{
+			return false;
+		}
+		UE::Interchange::FAttributeStorage::TAttributeHandle<int32> Handle = Attributes.GetAttributeHandle<int32>(UE::Interchange::FSkeletalMeshNodeStaticData::LodDataIdsCountKey());
+		if (!Handle.IsValid())
+		{
+			return false;
+		}
+		Handle.Get(LodCount);
+		bool DecrementKey = false;
+		for (int32 LodIndex = 0; LodIndex < LodCount; ++LodIndex)
+		{
+			FString DepIndexKeyString = UE::Interchange::FSkeletalMeshNodeStaticData::LodDataIdsBaseKey().ToString() + FString::FromInt(LodIndex);
+			UE::Interchange::FAttributeKey DepIndexKey(*DepIndexKeyString);
+			UE::Interchange::FAttributeStorage::TAttributeHandle<FName> HandleMesh = Attributes.GetAttributeHandle<FName>(DepIndexKey);
+			if (!HandleMesh.IsValid())
+			{
+				continue;
+			}
+			FName MeshUniqueID;
+			HandleMesh.Get(MeshUniqueID);
+			if (MeshUniqueID == LodDataUniqueId)
+			{
+				//Remove this entry
+				Attributes.UnregisterAttribute(DepIndexKey);
+				Handle.Set(LodCount - 1);
+				//We have to rename the key for all the next item
+				DecrementKey = true;;
+			}
+			else if (DecrementKey)
+			{
+				FString NewDepIndexKeyString = UE::Interchange::FSkeletalMeshNodeStaticData::LodDataIdsBaseKey().ToString() + FString::FromInt(LodIndex - 1);
+				UE::Interchange::FAttributeKey NewDepIndexKey(*NewDepIndexKeyString);
+				UE::Interchange::EAttributeStorageResult UnregisterResult = Attributes.UnregisterAttribute(DepIndexKey);
+				if (UE::Interchange::IsAttributeStorageResultSuccess(UnregisterResult))
+				{
+					UE::Interchange::EAttributeStorageResult RegisterResult = Attributes.RegisterAttribute<FName>(NewDepIndexKey, MeshUniqueID);
+					if (!UE::Interchange::IsAttributeStorageResultSuccess(RegisterResult))
+					{
+						UE::Interchange::LogAttributeStorageErrors(RegisterResult, TEXT("UInterchangeSkeletalMeshNode.RemoveLodDataUniqueId"), NewDepIndexKey);
+					}
+				}
+				else
+				{
+					UE::Interchange::LogAttributeStorageErrors(UnregisterResult, TEXT("UInterchangeSkeletalMeshNode.RemoveLodDataUniqueId"), DepIndexKey);
+				}
+
+				//Avoid doing more code in the for since the HandleMesh is now invalid
+				continue;
+			}
+		}
+		return true;
+	}
+private:
+
+	void FillAssetClassFromAttribute()
+	{
+#if WITH_ENGINE
+		FString OperationName = GetTypeName() + TEXT(".GetAssetClassName");
+		FString ClassName;
+		InterchangePrivateNodeBase::GetCustomAttribute<FString>(Attributes, ClassNameAttributeKey, OperationName, ClassName);
+		if (ClassName.Equals(USkeletalMesh::StaticClass()->GetName()))
+		{
+			AssetClass = USkeletalMesh::StaticClass();
+			bIsNodeClassInitialized = true;
+		}
+#endif
+	}
+
+	bool SetNodeClassFromClassAttribute()
+	{
+		if (!bIsNodeClassInitialized)
+		{
+			FillAssetClassFromAttribute();
+		}
+		return bIsNodeClassInitialized;
+	}
+
+	bool IsEditorOnlyDataDefined()
+	{
+#if WITH_EDITORONLY_DATA
+		return true;
+#else
+		return false;
+#endif
+	}
+
+	const UE::Interchange::FAttributeKey ClassNameAttributeKey = UE::Interchange::FAttributeKey(TEXT("__ClassTypeAttribute__"));
+
+protected:
+#if WITH_ENGINE
+	TSubclassOf<USkeletalMesh> AssetClass = nullptr;
+#endif
+	bool bIsNodeClassInitialized = false;
+};
diff --git a/Engine/Plugins/Runtime/Interchange/Source/Nodes/Public/InterchangeSkeletonNode.h b/Engine/Plugins/Runtime/Interchange/Source/Nodes/Public/InterchangeSkeletonNode.h
new file mode 100644
--- /dev/null
+++ b/Engine/Plugins/Runtime/Interchange/Source/Nodes/Public/InterchangeSkeletonNode.h
@@ -0,0 +1,141 @@
+// Copyright Epic Games, Inc. All Rights Reserved.
+
+#pragma once
+
+#include "CoreMinimal.h"
+#include "Nodes/InterchangeBaseNode.h"
+
+#if WITH_ENGINE
+#include "Animation/Skeleton.h"
+#endif
+
+#include "InterchangeSkeletonNode.generated.h"
+
+UCLASS(BlueprintType)
+class INTERCHANGENODEPLUGIN_API UInterchangeSkeletonNode : public UInterchangeBaseNode
+{
+	GENERATED_BODY()
+
+public:
+	UInterchangeSkeletonNode()
+	:UInterchangeBaseNode()
+	{
+#if WITH_ENGINE
+		AssetClass = nullptr;
+#endif
+	}
+
+	/**
+	 * Initialize node data
+	 * @param: UniqueID - The uniqueId for this node
+	 * @param DisplayLabel - The name of the node
+	 * @param InAssetClass - The class the Skeleton factory will create for this node.
+	 */
+	UFUNCTION(BlueprintCallable, Category = "Interchange | Node | Skeleton")
+	void InitializeSkeletonNode(const FName& UniqueID, const FName& DisplayLabel, const FString& InAssetClass)
+	{
+		bIsNodeClassInitialized = false;
+		InitializeNode(UniqueID, DisplayLabel);
+		FString OperationName = GetTypeName() + TEXT(".SetAssetClassName");
+		InterchangePrivateNodeBase::SetCustomAttribute<FString>(Attributes, ClassNameAttributeKey, OperationName, InAssetClass);
+		FillAssetClassFromAttribute();
+	}
+
+	virtual void Serialize(FArchive& Ar) override
+	{
+		Super::Serialize(Ar);
+#if WITH_ENGINE
+		if(Ar.IsLoading())
+		{
+			//Make sure the class is properly set when we compile with engine, this will set the
+			//bIsNodeClassInitialized to true.
+			SetNodeClassFromClassAttribute();
+		}
+#endif //#if WITH_ENGINE
+	}
+
+	/**
+	 * Return the node type name of the class, we use this when reporting error
+	 */
+	virtual FString GetTypeName() const override
+	{
+		const FString TypeName = TEXT("SkeletonNode");
+		return TypeName;
+	}
+
+	/** Get the class this node want to create */
+	UFUNCTION(BlueprintCallable, Category = "Interchange | Node | Skeleton")
+	virtual class UClass* GetAssetClass() const override
+	{
+		ensure(bIsNodeClassInitialized);
+#if WITH_ENGINE
+		return AssetClass.Get() != nullptr ? AssetClass.Get() : USkeleton::StaticClass();
+#else
+		return nullptr;
+#endif
+	}
+
+	virtual FGuid GetHash() const override
+	{
+		return Attributes.GetStorageHash();
+	}
+
+public:
+	/** Return false if the Attribute was not set previously.*/
+	UFUNCTION(BlueprintCallable, Category = "Interchange | Node | Skeleton")
+	bool GetCustomRootJointID(FName& AttributeValue) const
+	{
+		IMPLEMENT_NODE_ATTRIBUTE_GETTER(RootJointID, FName);
+	}
+
+	UFUNCTION(BlueprintCallable, Category = "Interchange | Node | Skeleton")
+	bool SetCustomRootJointID(const FName& AttributeValue)
+	{
+		IMPLEMENT_NODE_ATTRIBUTE_SETTER_NODELEGATE(RootJointID, FName)
+	}
+
+private:
+
+	void FillAssetClassFromAttribute()
+	{
+#if WITH_ENGINE
+		FString OperationName = GetTypeName() + TEXT(".GetAssetClassName");
+		FString ClassName;
+		InterchangePrivateNodeBase::GetCustomAttribute<FString>(Attributes, ClassNameAttributeKey, OperationName, ClassName);
+		if (ClassName.Equals(USkeleton::StaticClass()->GetName()))
+		{
+			AssetClass = USkeleton::StaticClass();
+			bIsNodeClassInitialized = true;
+		}
+#endif
+	}
+
+	bool SetNodeClassFromClassAttribute()
+	{
+		if (!bIsNodeClassInitialized)
+		{
+			FillAssetClassFromAttribute();
+		}
+		return bIsNodeClassInitialized;
+	}
+
+	bool IsEditorOnlyDataDefined()
+	{
+#if WITH_EDITORONLY_DATA
+		return true;
+#else
+		return false;
+#endif
+	}
+
+	const UE::Interchange::FAttributeKey ClassNameAttributeKey = UE::Interchange::FAttributeKey(TEXT("__ClassTypeAttribute__"));
+
+	//Skeleton
+	const UE::Interchange::FAttributeKey Macro_CustomRootJointIDKey = UE::Interchange::FAttributeKey(TEXT("RootJointID"));
+	
+protected:
+#if WITH_ENGINE
+	TSubclassOf<USkeleton> AssetClass = nullptr;
+#endif
+	bool bIsNodeClassInitialized = false;
+};
diff --git a/Engine/Plugins/Runtime/Interchange/Source/Parsers/Fbx/InterchangeFbxParser.Build.cs b/Engine/Plugins/Runtime/Interchange/Source/Parsers/Fbx/InterchangeFbxParser.Build.cs
--- a/Engine/Plugins/Runtime/Interchange/Source/Parsers/Fbx/InterchangeFbxParser.Build.cs
+++ b/Engine/Plugins/Runtime/Interchange/Source/Parsers/Fbx/InterchangeFbxParser.Build.cs
@@ -23,6 +23,9 @@ namespace UnrealBuildTool.Rules
 				new string[]
 				{
 					"InterchangeNodePlugin",
+					"MeshDescription",
+					"SkeletalMeshDescription",
+					"StaticMeshDescription"
 				}
 			);
 
diff --git a/Engine/Plugins/Runtime/Interchange/Source/Parsers/Fbx/Private/FbxAPI.cpp b/Engine/Plugins/Runtime/Interchange/Source/Parsers/Fbx/Private/FbxAPI.cpp
--- a/Engine/Plugins/Runtime/Interchange/Source/Parsers/Fbx/Private/FbxAPI.cpp
+++ b/Engine/Plugins/Runtime/Interchange/Source/Parsers/Fbx/Private/FbxAPI.cpp
@@ -4,9 +4,12 @@
 
 #include "CoreMinimal.h"
 #include "FbxConvert.h"
+#include "FbxHelper.h"
 #include "FbxInclude.h"
 #include "FbxMaterial.h"
 #include "FbxScene.h"
+#include "FbxSkeletalMesh.h"
+#include "FbxSkeleton.h"
 #include "InterchangeMaterialNode.h"
 #include "InterchangeTextureNode.h"
 #include "Nodes/InterchangeBaseNodeContainer.h"
@@ -17,8 +20,14 @@ namespace UE
 	{
 		namespace Private
 		{
+			FFbxParser::~FFbxParser()
+			{
+				PayloadContexts.Empty();
+			}
+
 			bool FFbxParser::LoadFbxFile(const FString& Filename, TArray<FString>& JSonErrorMessages)
 			{
+				SourceFilename = Filename;
 				int32 SDKMajor, SDKMinor, SDKRevision;
 				
 				//The first thing to do is to create the FBX Manager which is the object allocator for almost all the classes in the SDK
@@ -43,6 +52,8 @@ namespace UE
 				ios->SetBoolProp(IMP_TAKE, true);
 				SDKManager->SetIOSettings(ios);
 
+				SDKGeometryConverter = new FbxGeometryConverter(SDKManager);
+
 				//Create an FBX scene. This object holds most objects imported/exported from/to files.
 				SDKScene = FbxScene::Create(SDKManager, "My Scene");
 				if (!SDKScene)
@@ -82,8 +93,23 @@ namespace UE
 			void FFbxParser::FillContainerWithFbxScene(UInterchangeBaseNodeContainer& NodeContainer, TArray<FString>& JSonErrorMessages)
 			{
 				FFbxMaterial::AddAllSceneMaterials(SDKScene, NodeContainer, JSonErrorMessages);
-				FFbxScene::AddHierarchy(SDKScene, NodeContainer, JSonErrorMessages);
+				FFbxSkeleton::AddAllSceneSkeletons(SDKScene, NodeContainer, JSonErrorMessages, SourceFilename);
+				FFbxSkeletalMesh::AddAllSceneSkeletalMeshes(SDKScene, SDKGeometryConverter, NodeContainer, JSonErrorMessages, PayloadContexts, SourceFilename);
+				//FFbxScene::AddHierarchy(SDKScene, NodeContainer, JSonErrorMessages);
+			}
+
+			bool FFbxParser::FetchPayloadData(const FString& PayloadKey, const FString& PayloadFilepath, TArray<FString>& JSonErrorMessages)
+			{
+				if (!PayloadContexts.Contains(PayloadKey))
+				{
+					JSonErrorMessages.Add(TEXT("{\"Msg\" : {\"Type\" : \"Error\",\n\"Msg\" : \"Cannot retrieve payload, payload key doesnt have any context.\"}}"));
+					return false;
+				}
+
+				TSharedPtr<FPayloadContextBase>& PayloadContext = PayloadContexts.FindChecked(PayloadKey);
+				return PayloadContext->FetchPayloadToFile(PayloadFilepath, JSonErrorMessages);
 			}
+
 		} //ns Private
 	} //ns Interchange
 } //ns UE
diff --git a/Engine/Plugins/Runtime/Interchange/Source/Parsers/Fbx/Private/FbxAPI.h b/Engine/Plugins/Runtime/Interchange/Source/Parsers/Fbx/Private/FbxAPI.h
--- a/Engine/Plugins/Runtime/Interchange/Source/Parsers/Fbx/Private/FbxAPI.h
+++ b/Engine/Plugins/Runtime/Interchange/Source/Parsers/Fbx/Private/FbxAPI.h
@@ -8,6 +8,16 @@
 #define FBX_METADATA_PREFIX TEXT("FBX.")
 #define INVALID_UNIQUE_ID 0xFFFFFFFFFFFFFFFF
 
+namespace UE
+{
+	namespace Interchange
+	{
+		namespace Private
+		{
+			class FPayloadContextBase;
+		}
+	}
+}
 class UInterchangeBaseNodeContainer;
 
 namespace UE
@@ -19,17 +29,26 @@ namespace UE
 			class FFbxParser
 			{
 			public:
+				~FFbxParser();
+
 				/* Load an fbx file into the fbx sdk, return false if the file could not be load. */
 				bool LoadFbxFile(const FString& Filename, TArray<FString>& JsonErrorMessages);
 
 				/* Extract the fbx data from the sdk into our node container */
 				void FillContainerWithFbxScene(UInterchangeBaseNodeContainer& NodeContainer, TArray<FString>& JsonErrorMessages);
+
+				/* Extract the fbx data from the sdk into our node container */
+				bool FetchPayloadData(const FString& PayloadKey, const FString& PayloadFilepath, TArray<FString>& JSonErrorMessages);
 			
 			private:
 				FbxManager* SDKManager = nullptr;
 				FbxScene* SDKScene = nullptr;
 				FbxImporter* SDKImporter = nullptr;
-				//FbxGeometryConverter* SDKGeometryConverter = nullptr;
+				FbxGeometryConverter* SDKGeometryConverter = nullptr;
+				FString SourceFilename;
+
+				TMap<FString, TSharedPtr<FPayloadContextBase>> PayloadContexts;
+				
 			};
 		}//ns Private
 	}//ns Interchange
diff --git a/Engine/Plugins/Runtime/Interchange/Source/Parsers/Fbx/Private/FbxConvert.cpp b/Engine/Plugins/Runtime/Interchange/Source/Parsers/Fbx/Private/FbxConvert.cpp
--- a/Engine/Plugins/Runtime/Interchange/Source/Parsers/Fbx/Private/FbxConvert.cpp
+++ b/Engine/Plugins/Runtime/Interchange/Source/Parsers/Fbx/Private/FbxConvert.cpp
@@ -5,6 +5,8 @@
 #include "CoreMinimal.h"
 #include "FbxInclude.h"
 
+#define CONVERT_TO_FRONT_X 0
+
 namespace UE
 {
 	namespace Interchange
@@ -42,9 +44,13 @@ namespace UE
 
 
 				//UE is: z up, front x, left handed
-				FbxAxisSystem::EUpVector UpVector = FbxAxisSystem::eZAxis;
+				FbxAxisSystem::EUpVector UpVector = FbxAxisSystem::EUpVector::eZAxis;
+#if CONVERT_TO_FRONT_X
 				FbxAxisSystem::EFrontVector FrontVector = (FbxAxisSystem::EFrontVector)FbxAxisSystem::eParityEven;
-				FbxAxisSystem::ECoordSystem CoordSystem = FbxAxisSystem::eLeftHanded;
+#else
+				FbxAxisSystem::EFrontVector FrontVector = (FbxAxisSystem::EFrontVector) - FbxAxisSystem::eParityOdd;
+#endif //CONVERT_TO_FRONT_X
+				FbxAxisSystem::ECoordSystem CoordSystem = FbxAxisSystem::ECoordSystem::eRightHanded;
 				FbxAxisSystem UnrealImportAxis(UpVector, FrontVector, CoordSystem);
 
 				if (FileAxisSystem != UnrealImportAxis)
@@ -62,23 +68,85 @@ namespace UE
 				SDKScene->GetAnimationEvaluator()->Reset();
 			}
 
-			FTransform FFbxConvert::GetTransform(const FbxAMatrix& Matrix)
+			FTransform FFbxConvert::ConvertTransform(const FbxAMatrix& Matrix)
+			{
+				FTransform Out;
+
+				FQuat Rotation = ConvertRotToQuat(Matrix.GetQ());
+				FVector Origin = ConvertPos(Matrix.GetT());
+				FVector Scale = ConvertScale(Matrix.GetS());
+
+				Out.SetTranslation(Origin);
+				Out.SetScale3D(Scale);
+				Out.SetRotation(Rotation);
+
+				return Out;
+			}
+
+			FMatrix FFbxConvert::ConvertMatrix(const FbxAMatrix& Matrix)
+			{
+				FMatrix UEMatrix;
+
+				for (int i = 0; i < 4; ++i)
+				{
+					const FbxVector4 Row = Matrix.GetRow(i);
+					if (i == 1)
+					{
+						UEMatrix.M[i][0] = -Row[0];
+						UEMatrix.M[i][1] = Row[1];
+						UEMatrix.M[i][2] = -Row[2];
+						UEMatrix.M[i][3] = -Row[3];
+					}
+					else
+					{
+						UEMatrix.M[i][0] = Row[0];
+						UEMatrix.M[i][1] = -Row[1];
+						UEMatrix.M[i][2] = Row[2];
+						UEMatrix.M[i][3] = Row[3];
+					}
+				}
+
+				return UEMatrix;
+			}
+
+			FQuat FFbxConvert::ConvertRotToQuat(FbxQuaternion Quaternion)
+			{
+				FQuat UnrealQuat;
+				UnrealQuat.X = Quaternion[0];
+				UnrealQuat.Y = -Quaternion[1];
+				UnrealQuat.Z = Quaternion[2];
+				UnrealQuat.W = -Quaternion[3];
+
+				return UnrealQuat;
+			}
+
+			FRotator FFbxConvert::ConvertEuler(FbxDouble3 Euler)
+			{
+				return FRotator::MakeFromEuler(FVector(Euler[0], -Euler[1], Euler[2]));
+			}
+
+			FVector FFbxConvert::ConvertScale(FbxVector4 Vector)
+			{
+				FVector Out;
+				Out[0] = Vector[0];
+				Out[1] = Vector[1];
+				Out[2] = Vector[2];
+				return Out;
+			}
+
+			FRotator FFbxConvert::ConvertRotation(FbxQuaternion Quaternion)
+			{
+				FRotator Out(ConvertRotToQuat(Quaternion));
+				return Out;
+			}
+			FVector FFbxConvert::ConvertPos(const FbxVector4& Vector)
+			{
+				return FVector(Vector[0], -Vector[1], Vector[2]);
+			}
+
+			FVector FFbxConvert::ConvertDir(const FbxVector4& Vector)
 			{
-				FTransform Transform = FTransform::Identity;
-				FbxVector4 FbxLocalT = Matrix.GetT();
-				FbxVector4 FbxLocalS = Matrix.GetS();
-				FbxQuaternion FbxLocalQ = Matrix.GetQ();
-				FVector LocalT = FVector(FbxLocalT[0], FbxLocalT[1], FbxLocalT[2]);
-				Transform.SetTranslation(LocalT);
-				FVector LocalS = FVector(FbxLocalS[0], FbxLocalS[1], FbxLocalS[2]);
-				Transform.SetScale3D(LocalS);
-				FQuat LocalQ;
-				LocalQ.X = FbxLocalQ[0];
-				LocalQ.Y = FbxLocalQ[1];
-				LocalQ.Z = FbxLocalQ[2];
-				LocalQ.W = FbxLocalQ[3];
-				Transform.SetRotation(LocalQ);
-				return Transform;
+				return FVector(Vector[0], -Vector[1], Vector[2]);
 			}
 
 			ANSICHAR* FFbxConvert::MakeName(const ANSICHAR* Name)
diff --git a/Engine/Plugins/Runtime/Interchange/Source/Parsers/Fbx/Private/FbxConvert.h b/Engine/Plugins/Runtime/Interchange/Source/Parsers/Fbx/Private/FbxConvert.h
--- a/Engine/Plugins/Runtime/Interchange/Source/Parsers/Fbx/Private/FbxConvert.h
+++ b/Engine/Plugins/Runtime/Interchange/Source/Parsers/Fbx/Private/FbxConvert.h
@@ -21,7 +21,14 @@ namespace UE
 				/**
 				 * Create a UE transform from a fbx matrix
 				 */
-				static FTransform GetTransform(const FbxAMatrix& Matrix);
+				static FTransform ConvertTransform(const FbxAMatrix& Matrix);
+				static FMatrix ConvertMatrix(const FbxAMatrix& Matrix);
+				static FQuat ConvertRotToQuat(FbxQuaternion Quaternion);
+				static FRotator ConvertEuler(FbxDouble3 Euler);
+				static FVector ConvertScale(FbxVector4 Vector);
+				static FRotator ConvertRotation(FbxQuaternion Quaternion);
+				static FVector ConvertPos(const FbxVector4& Vector);
+				static FVector ConvertDir(const FbxVector4& Vector);
 
 				/** Transform Conversion API End */
 				//////////////////////////////////////////////////////////////////////////
diff --git a/Engine/Plugins/Runtime/Interchange/Source/Parsers/Fbx/Private/FbxHelper.cpp b/Engine/Plugins/Runtime/Interchange/Source/Parsers/Fbx/Private/FbxHelper.cpp
new file mode 100644
--- /dev/null
+++ b/Engine/Plugins/Runtime/Interchange/Source/Parsers/Fbx/Private/FbxHelper.cpp
@@ -0,0 +1,971 @@
+// Copyright Epic Games, Inc. All Rights Reserved.
+
+#include "FbxHelper.h"
+
+#include "CoreMinimal.h"
+#include "FbxConvert.h"
+#include "FbxInclude.h"
+
+#define GeneratedLODNameSuffix "_GeneratedLOD_"
+
+namespace UE
+{
+	namespace Interchange
+	{
+		namespace Private
+		{
+			FString FFbxHelper::GetFbxObjectName(const FbxObject* Object)
+			{
+				if (!Object)
+				{
+					return FString();
+				}
+				FString ObjName = UTF8_TO_TCHAR(FFbxConvert::MakeName(Object->GetName()));
+				if (ObjName.Equals(TEXT("none"), ESearchCase::IgnoreCase))
+				{
+					//Replace None by Null because None clash with NAME_None and the create asset will instead call the object ClassName_X
+					ObjName = TEXT("Null");
+				}
+				return ObjName;
+			}
+
+			FString FFbxHelper::GetFbxNodeHierarchyName(const FbxNode* Node)
+			{
+				if (!Node)
+				{
+					return FString();
+				}
+				TArray<FString> UniqueIDTokens;
+				const FbxNode* ParentNode = Node;
+				while (ParentNode)
+				{
+					UniqueIDTokens.Add(GetFbxObjectName(ParentNode));
+					ParentNode = ParentNode->GetParent();
+				}
+				FString UniqueID;
+				for (int32 TokenIndex = UniqueIDTokens.Num() - 1; TokenIndex >= 0; TokenIndex--)
+				{
+					UniqueID += UniqueIDTokens[TokenIndex];
+					if (TokenIndex > 0)
+					{
+						UniqueID += TEXT(".");
+					}
+				}
+				return UniqueID;
+			}
+
+			void FFbxHelper::FindSkeletalMeshes(FbxScene* SDKScene, TArray< TArray<FbxNode*> >& outSkelMeshArray, bool bCombineSkeletalMesh, bool bForceFindRigid)
+			{
+				TArray<FbxNode*> SkeletonArray;
+				FbxNode* RootNode = SDKScene->GetRootNode();
+
+				// a) find skeletal meshes
+
+				RecursiveFindFbxSkelMesh(SDKScene, RootNode, outSkelMeshArray, SkeletonArray);
+				// for skeletal mesh, we convert the skeleton system to skeleton
+				// in less we recognize bone mesh as rigid mesh if they are textured
+				for (int32 SkelIndex = 0; SkelIndex < SkeletonArray.Num(); SkelIndex++)
+				{
+					const bool bImportMeshesInBoneHierarchy = true;
+					RecursiveFixSkeleton(SDKScene, SkeletonArray[SkelIndex], outSkelMeshArray[SkelIndex], bImportMeshesInBoneHierarchy);
+				}
+
+
+
+				// b) find rigid mesh
+
+				// If we are attempting to import a skeletal mesh but we have no hierarchy attempt to find a rigid mesh.
+				if (bForceFindRigid || outSkelMeshArray.Num() == 0)
+				{
+					RecursiveFindRigidMesh(SDKScene, RootNode, outSkelMeshArray, SkeletonArray);
+					if (bForceFindRigid)
+					{
+						//Cleanup the rigid mesh, We want to remove any real static mesh from the outSkelMeshArray
+						//Any non skinned mesh that contain no animation should be part of this array.
+						int32 AnimStackCount = SDKScene->GetSrcObjectCount<FbxAnimStack>();
+						TArray<int32> SkeletalMeshArrayToRemove;
+						for (int32 i = 0; i < outSkelMeshArray.Num(); i++)
+						{
+							bool bIsValidSkeletal = false;
+							TArray<FbxNode*>& NodeArray = outSkelMeshArray[i];
+							for (FbxNode* InspectedNode : NodeArray)
+							{
+								FbxMesh* Mesh = InspectedNode->GetMesh();
+
+								FbxLODGroup* LodGroup = InspectedNode->GetLodGroup();
+								if (LodGroup != nullptr)
+								{
+									FbxNode* SkelMeshNode = FFbxHelper::FindLODGroupNode(InspectedNode, 0, nullptr);
+									if (SkelMeshNode != nullptr)
+									{
+										Mesh = SkelMeshNode->GetMesh();
+									}
+								}
+
+								if (Mesh == nullptr)
+								{
+									continue;
+								}
+								if (Mesh->GetDeformerCount(FbxDeformer::eSkin) > 0)
+								{
+									bIsValidSkeletal = true;
+									break;
+								}
+								//If there is some anim object we count this as a valid skeletal mesh imported as rigid mesh
+								for (int32 AnimStackIndex = 0; AnimStackIndex < AnimStackCount; AnimStackIndex++)
+								{
+									FbxAnimStack* CurAnimStack = SDKScene->GetSrcObject<FbxAnimStack>(AnimStackIndex);
+									// set current anim stack
+
+									SDKScene->SetCurrentAnimationStack(CurAnimStack);
+
+									FbxTimeSpan AnimTimeSpan(FBXSDK_TIME_INFINITE, FBXSDK_TIME_MINUS_INFINITE);
+									InspectedNode->GetAnimationInterval(AnimTimeSpan, CurAnimStack);
+
+									if (AnimTimeSpan.GetDuration() > 0)
+									{
+										bIsValidSkeletal = true;
+										break;
+									}
+								}
+								if (bIsValidSkeletal)
+								{
+									break;
+								}
+							}
+							if (!bIsValidSkeletal)
+							{
+								SkeletalMeshArrayToRemove.Add(i);
+							}
+						}
+						for (int32 i = SkeletalMeshArrayToRemove.Num() - 1; i >= 0; --i)
+						{
+							if (!SkeletalMeshArrayToRemove.IsValidIndex(i) || !outSkelMeshArray.IsValidIndex(SkeletalMeshArrayToRemove[i]))
+								continue;
+							int32 IndexToRemove = SkeletalMeshArrayToRemove[i];
+							outSkelMeshArray[IndexToRemove].Empty();
+							outSkelMeshArray.RemoveAt(IndexToRemove);
+						}
+					}
+				}
+				//Empty the skeleton array
+				SkeletonArray.Empty();
+
+
+				if (bCombineSkeletalMesh)
+				{
+					//Merge all the skeletal mesh arrays into one combine mesh
+					TArray<FbxNode*> CombineNodes;
+					for (TArray<FbxNode*>& Parts : outSkelMeshArray)
+					{
+						CombineNodes.Append(Parts);
+					}
+					outSkelMeshArray.Empty(1);
+					outSkelMeshArray.Add(CombineNodes);
+				}
+			}
+
+			void FFbxHelper::RecursiveGetAllMeshNode(TArray<FbxNode*> & OutAllNode, FbxNode * Node)
+			{
+				if (Node == nullptr)
+				{
+					return;
+				}
+
+				if (Node->GetMesh() != nullptr)
+				{
+					OutAllNode.Add(Node);
+					return;
+				}
+
+				//Look if its a generated LOD
+				FString FbxGeneratedNodeName = UTF8_TO_TCHAR(Node->GetName());
+				if (FbxGeneratedNodeName.Contains(TEXT(GeneratedLODNameSuffix)))
+				{
+					FString SuffixSearch = TEXT(GeneratedLODNameSuffix);
+					int32 SuffixIndex = FbxGeneratedNodeName.Find(SuffixSearch, ESearchCase::CaseSensitive, ESearchDir::FromEnd);
+					SuffixIndex += SuffixSearch.Len();
+					FString LODXNumber = FbxGeneratedNodeName.RightChop(SuffixIndex).Left(1);
+					if (LODXNumber.IsNumeric())
+					{
+						OutAllNode.Add(Node);
+						return;
+					}
+				}
+
+				for (int32 ChildIndex = 0; ChildIndex < Node->GetChildCount(); ++ChildIndex)
+				{
+					RecursiveGetAllMeshNode(OutAllNode, Node->GetChild(ChildIndex));
+				}
+			}
+
+			void FFbxHelper::FindAllLODGroupNode(TArray<FbxNode*> & OutNodeInLod, FbxNode * NodeLodGroup, int32 LodIndex)
+			{
+				check(NodeLodGroup->GetChildCount() >= LodIndex);
+				FbxNode* ChildNode = NodeLodGroup->GetChild(LodIndex);
+				if (ChildNode == nullptr)
+				{
+					return;
+				}
+				RecursiveGetAllMeshNode(OutNodeInLod, ChildNode);
+			}
+
+			bool FFbxHelper::IsUnrealBone(FbxNode * Link)
+			{
+				FbxNodeAttribute* Attr = Link->GetNodeAttribute();
+				if (Attr)
+				{
+					FbxNodeAttribute::EType AttrType = Attr->GetAttributeType();
+					if (AttrType == FbxNodeAttribute::eSkeleton ||
+						AttrType == FbxNodeAttribute::eMesh ||
+						AttrType == FbxNodeAttribute::eNull)
+					{
+						return true;
+					}
+				}
+
+				return false;
+			}
+
+			void FFbxHelper::RecursiveBuildSkeleton(FbxNode * Link, TArray<FbxNode*> & OutSortedLinks)
+			{
+				if (IsUnrealBone(Link))
+				{
+					OutSortedLinks.Add(Link);
+					int32 ChildIndex;
+					for (ChildIndex = 0; ChildIndex < Link->GetChildCount(); ChildIndex++)
+					{
+						RecursiveBuildSkeleton(Link->GetChild(ChildIndex), OutSortedLinks);
+					}
+				}
+			}
+
+			bool FFbxHelper::RetrievePoseFromBindPose(FbxScene * SDKScene, const TArray<FbxNode*> & NodeArray, FbxArray<FbxPose*> & PoseArray)
+			{
+				const int32 PoseCount = SDKScene->GetPoseCount();
+				for (int32 PoseIndex = 0; PoseIndex < PoseCount; PoseIndex++)
+				{
+					FbxPose* CurrentPose = SDKScene->GetPose(PoseIndex);
+
+					// current pose is bind pose, 
+					if (CurrentPose && CurrentPose->IsBindPose())
+					{
+						// IsValidBindPose doesn't work reliably
+						// It checks all the parent chain(regardless root given), and if the parent doesn't have correct bind pose, it fails
+						// It causes more false positive issues than the real issue we have to worry about
+						// If you'd like to try this, set CHECK_VALID_BIND_POSE to 1, and try the error message
+						// when Autodesk fixes this bug, then we might be able to re-open this
+						FString PoseName = CurrentPose->GetName();
+						// all error report status
+						FbxStatus Status;
+
+						// it does not make any difference of checking with different node
+						// it is possible pose 0 -> node array 2, but isValidBindPose function returns true even with node array 0
+						for (auto Current : NodeArray)
+						{
+							FString CurrentName = Current->GetName();
+							NodeList pMissingAncestors, pMissingDeformers, pMissingDeformersAncestors, pWrongMatrices;
+
+							if (CurrentPose->IsValidBindPoseVerbose(Current, pMissingAncestors, pMissingDeformers, pMissingDeformersAncestors, pWrongMatrices, 0.0001, &Status))
+							{
+								PoseArray.Add(CurrentPose);
+								//UE_LOG(LogFbx, Log, TEXT("Valid bind pose for Pose (%s) - %s"), *PoseName, *FString(Current->GetName()));
+								break;
+							}
+							else
+							{
+								// first try to fix up
+								// add missing ancestors
+								for (int i = 0; i < pMissingAncestors.GetCount(); i++)
+								{
+									FbxAMatrix mat = pMissingAncestors.GetAt(i)->EvaluateGlobalTransform(FBXSDK_TIME_ZERO);
+									CurrentPose->Add(pMissingAncestors.GetAt(i), mat);
+								}
+
+								pMissingAncestors.Clear();
+								pMissingDeformers.Clear();
+								pMissingDeformersAncestors.Clear();
+								pWrongMatrices.Clear();
+
+								// check it again
+								if (CurrentPose->IsValidBindPose(Current))
+								{
+									PoseArray.Add(CurrentPose);
+									//UE_LOG(LogFbx, Log, TEXT("Valid bind pose for Pose (%s) - %s"), *PoseName, *FString(Current->GetName()));
+									break;
+								}
+								else
+								{
+									// first try to find parent who is null group and see if you can try test it again
+									FbxNode* ParentNode = Current->GetParent();
+									while (ParentNode)
+									{
+										FbxNodeAttribute* Attr = ParentNode->GetNodeAttribute();
+										if (Attr && Attr->GetAttributeType() == FbxNodeAttribute::eNull)
+										{
+											// found it 
+											break;
+										}
+
+										// find next parent
+										ParentNode = ParentNode->GetParent();
+									}
+
+									if (ParentNode && CurrentPose->IsValidBindPose(ParentNode))
+									{
+										PoseArray.Add(CurrentPose);
+										//UE_LOG(LogFbx, Log, TEXT("Valid bind pose for Pose (%s) - %s"), *PoseName, *FString(Current->GetName()));
+										break;
+									}
+									else
+									{
+										FString ErrorString = Status.GetErrorString();
+										// 											if (!GIsAutomationTesting)
+										// 												UE_LOG(LogFbx, Warning, TEXT("Not valid bind pose for Pose (%s) - Node %s : %s"), *PoseName, *FString(Current->GetName()), *ErrorString);
+									}
+								}
+							}
+						}
+					}
+				}
+
+				return (PoseArray.Size() > 0);
+			}
+
+			FbxNode* FFbxHelper::GetRootSkeleton(FbxScene * SDKScene, FbxNode * Link)
+			{
+				FbxNode* RootBone = Link;
+
+				// get Unreal skeleton root
+				// mesh and dummy are used as bone if they are in the skeleton hierarchy
+				while (RootBone && RootBone->GetParent())
+				{
+					bool bIsBlenderArmatureBone = false;
+					//TODO put back the creator code so we know if its blender
+					//if (FbxCreator == EFbxCreator::Blender)
+					{
+						//Hack to support armature dummy node from blender
+						//Users do not want the null attribute node named armature which is the parent of the real root bone in blender fbx file
+						//This is a hack since if a rigid mesh group root node is named "armature" it will be skip
+						const FString RootBoneParentName(RootBone->GetParent()->GetName());
+						FbxNode* GrandFather = RootBone->GetParent()->GetParent();
+						bIsBlenderArmatureBone = (GrandFather == nullptr || GrandFather == SDKScene->GetRootNode()) && (RootBoneParentName.Compare(TEXT("armature"), ESearchCase::IgnoreCase) == 0);
+					}
+
+					FbxNodeAttribute* Attr = RootBone->GetParent()->GetNodeAttribute();
+					if (Attr &&
+						(Attr->GetAttributeType() == FbxNodeAttribute::eMesh ||
+						(Attr->GetAttributeType() == FbxNodeAttribute::eNull && !bIsBlenderArmatureBone) ||
+						Attr->GetAttributeType() == FbxNodeAttribute::eSkeleton) &&
+						RootBone->GetParent() != SDKScene->GetRootNode())
+					{
+						// in some case, skeletal mesh can be ancestor of bones
+						// this avoids this situation
+						if (Attr->GetAttributeType() == FbxNodeAttribute::eMesh)
+						{
+							FbxMesh* Mesh = (FbxMesh*)Attr;
+							if (Mesh->GetDeformerCount(FbxDeformer::eSkin) > 0)
+							{
+								break;
+							}
+						}
+
+						RootBone = RootBone->GetParent();
+					}
+					else
+					{
+						break;
+					}
+				}
+
+				return RootBone;
+			}
+
+			void FFbxHelper::BuildSkeletonSystem(FbxScene * SDKScene, TArray<FbxCluster*> & ClusterArray, TArray<FbxNode*> & OutSortedLinks)
+			{
+				FbxNode* Link;
+				TArray<FbxNode*> RootLinks;
+				int32 ClusterIndex;
+				for (ClusterIndex = 0; ClusterIndex < ClusterArray.Num(); ClusterIndex++)
+				{
+					Link = ClusterArray[ClusterIndex]->GetLink();
+					if (Link)
+					{
+						Link = GetRootSkeleton(SDKScene, Link);
+						int32 LinkIndex;
+						for (LinkIndex = 0; LinkIndex < RootLinks.Num(); LinkIndex++)
+						{
+							if (Link == RootLinks[LinkIndex])
+							{
+								break;
+							}
+						}
+
+						// this link is a new root, add it
+						if (LinkIndex == RootLinks.Num())
+						{
+							RootLinks.Add(Link);
+						}
+					}
+				}
+
+				for (int32 LinkIndex = 0; LinkIndex < RootLinks.Num(); LinkIndex++)
+				{
+					RecursiveBuildSkeleton(RootLinks[LinkIndex], OutSortedLinks);
+				}
+			}
+
+			bool FFbxHelper::FindSkeletonJoints(FbxScene * SDKScene, TArray<FbxNode*> & NodeArray, TArray<FbxNode*> & SortedLinks, FbxArray<FbxAMatrix> & LocalsPerLink)
+			{
+				bool bIsARigidMesh = false;
+				FbxNode* Link = NULL;
+				FbxArray<FbxPose*> PoseArray;
+				TArray<FbxCluster*> ClusterArray;
+
+				if (NodeArray[0]->GetMesh()->GetDeformerCount(FbxDeformer::eSkin) == 0)
+				{
+					bIsARigidMesh = true;
+					Link = NodeArray[0];
+					RecursiveBuildSkeleton(GetRootSkeleton(SDKScene, Link), SortedLinks);
+				}
+				else
+				{
+					// get bindpose and clusters from FBX skeleton
+
+					// let's put the elements to their bind pose! (and we restore them after
+					// we have built the ClusterInformation.
+					int32 Default_NbPoses = SDKScene->GetFbxManager()->GetBindPoseCount(SDKScene);
+					// If there are no BindPoses, the following will generate them.
+					SDKScene->GetFbxManager()->CreateMissingBindPoses(SDKScene);
+
+					//if we created missing bind poses, update the number of bind poses
+					int32 NbPoses = SDKScene->GetFbxManager()->GetBindPoseCount(SDKScene);
+
+					if (NbPoses != Default_NbPoses)
+					{
+						//TODO output an error message to let the user know there is a missing bind pose when translating this fbx skeletalmesh
+						//Missing bind pose
+					}
+
+					//
+					// create the bones / skinning
+					//
+
+					for (int32 i = 0; i < NodeArray.Num(); i++)
+					{
+						FbxMesh* FbxMesh = NodeArray[i]->GetMesh();
+						const int32 SkinDeformerCount = FbxMesh->GetDeformerCount(FbxDeformer::eSkin);
+						for (int32 DeformerIndex = 0; DeformerIndex < SkinDeformerCount; DeformerIndex++)
+						{
+							FbxSkin* Skin = (FbxSkin*)FbxMesh->GetDeformer(DeformerIndex, FbxDeformer::eSkin);
+							for (int32 ClusterIndex = 0; ClusterIndex < Skin->GetClusterCount(); ClusterIndex++)
+							{
+								ClusterArray.Add(Skin->GetCluster(ClusterIndex));
+							}
+						}
+					}
+
+					if (ClusterArray.Num() == 0)
+					{
+						//TODO warn user we found a mesh with a deformer that is not bind to any skin cluster
+						return false;
+					}
+
+					// get bind pose
+					if (RetrievePoseFromBindPose(SDKScene, NodeArray, PoseArray) == false)
+					{
+						// 							if (!GIsAutomationTesting)
+						// 								UE_LOG(LogFbx, Warning, TEXT("Getting valid bind pose failed. Try to recreate bind pose"));
+													// if failed, delete bind pose, and retry.
+						const int32 PoseCount = SDKScene->GetPoseCount();
+						for (int32 PoseIndex = PoseCount - 1; PoseIndex >= 0; --PoseIndex)
+						{
+							FbxPose* CurrentPose = SDKScene->GetPose(PoseIndex);
+
+							// current pose is bind pose, 
+							if (CurrentPose && CurrentPose->IsBindPose())
+							{
+								SDKScene->RemovePose(PoseIndex);
+								CurrentPose->Destroy();
+							}
+						}
+
+						SDKScene->GetFbxManager()->CreateMissingBindPoses(SDKScene);
+						if (RetrievePoseFromBindPose(SDKScene, NodeArray, PoseArray) == false)
+						{
+							// 								if (!GIsAutomationTesting)
+							// 									UE_LOG(LogFbx, Warning, TEXT("Recreating bind pose failed."));
+						}
+						else
+						{
+							// 								if (!GIsAutomationTesting)
+							// 									UE_LOG(LogFbx, Warning, TEXT("Recreating bind pose succeeded."));
+						}
+					}
+
+					// recurse through skeleton and build ordered table
+					BuildSkeletonSystem(SDKScene, ClusterArray, SortedLinks);
+				}
+
+				// error check
+				// if no bond is found
+				if (SortedLinks.Num() == 0)
+				{
+					//						{
+					//							AddTokenizedErrorMessage(FTokenizedMessage::Create(EMessageSeverity::Warning, FText::Format(LOCTEXT("FbxSkeletaLMeshimport_NoBone", "'{0}' has no bones"), FText::FromString(NodeArray[0]->GetName()))), FFbxErrors::SkeletalMesh_NoBoneFound);
+					//						}
+					return false;
+				}
+
+				// if no bind pose is found
+				if (PoseArray.GetCount() == 0)
+				{
+					//TODO error to the user
+					//AddTokenizedErrorMessage(FTokenizedMessage::Create(EMessageSeverity::Warning, LOCTEXT("FbxSkeletaLMeshimport_MissingBindPose", "Could not find the bind pose.  It will use time 0 as bind pose.")), FFbxErrors::SkeletalMesh_InvalidBindPose);
+				}
+
+				int32 LinkIndex;
+
+				// Check for duplicate bone names and issue a warning if found
+				for (LinkIndex = 0; LinkIndex < SortedLinks.Num(); ++LinkIndex)
+				{
+					Link = SortedLinks[LinkIndex];
+
+					for (int32 AltLinkIndex = LinkIndex + 1; AltLinkIndex < SortedLinks.Num(); ++AltLinkIndex)
+					{
+						FbxNode* AltLink = SortedLinks[AltLinkIndex];
+
+						if (FCStringAnsi::Strcmp(Link->GetName(), AltLink->GetName()) == 0)
+						{
+							//								FString RawBoneName = UTF8_TO_TCHAR(Link->GetName());
+															//TODO notify the user that there is a duplicate bone name
+							// 									AddTokenizedErrorMessage(FTokenizedMessage::Create(EMessageSeverity::Error, FText::Format(LOCTEXT("Error_DuplicateBoneName", "Error, Could not import {0}.\nDuplicate bone name found ('{1}'). Each bone must have a unique name."),
+							// 															 FText::FromString(NodeArray[0]->GetName()), FText::FromString(RawBoneName))), FFbxErrors::SkeletalMesh_DuplicateBones);
+							return false;
+						}
+					}
+				}
+
+				FbxArray<FbxAMatrix> GlobalsPerLink;
+				GlobalsPerLink.Grow(SortedLinks.Num());
+				GlobalsPerLink[0].SetIdentity();
+				LocalsPerLink.Grow(SortedLinks.Num());
+				LocalsPerLink[0].SetIdentity();
+
+				bool GlobalLinkFoundFlag;
+
+				bool bAnyLinksNotInBindPose = false;
+				FString LinksWithoutBindPoses;
+
+				int32 RootIdx = INDEX_NONE;
+
+				for (LinkIndex = 0; LinkIndex < SortedLinks.Num(); LinkIndex++)
+				{
+					// Add a bone for each FBX Link
+					Link = SortedLinks[LinkIndex];
+					int32 ParentIndex = INDEX_NONE; // base value for root if no parent found
+
+					if (LinkIndex == 0)
+					{
+						RootIdx = LinkIndex;
+					}
+					else
+					{
+						const FbxNode* LinkParent = Link->GetParent();
+						// get the link parent index.
+						for (int32 ll = 0; ll < LinkIndex; ++ll) // <LinkIndex because parent is guaranteed to be before child in sortedLink
+						{
+							FbxNode* Otherlink = SortedLinks[ll];
+							if (Otherlink == LinkParent)
+							{
+								ParentIndex = ll;
+								break;
+							}
+						}
+
+						if (ParentIndex == INDEX_NONE)//We found another root inside the hierarchy, this is not supported
+						{
+							//AddTokenizedErrorMessage(FTokenizedMessage::Create(EMessageSeverity::Error, LOCTEXT("MultipleRootsFound", "Multiple roots are found in the bone hierarchy. We only support single root bone.")), FFbxErrors::SkeletalMesh_MultipleRoots);
+							return false;
+						}
+					}
+
+					GlobalLinkFoundFlag = false;
+					if (!bIsARigidMesh) //skeletal mesh
+					{
+						// there are some links, they have no cluster, but in bindpose
+						if (PoseArray.GetCount())
+						{
+							for (int32 PoseIndex = 0; PoseIndex < PoseArray.GetCount(); PoseIndex++)
+							{
+								int32 PoseLinkIndex = PoseArray[PoseIndex]->Find(Link);
+								if (PoseLinkIndex >= 0)
+								{
+									FbxMatrix NoneAffineMatrix = PoseArray[PoseIndex]->GetMatrix(PoseLinkIndex);
+									FbxAMatrix Matrix = *(FbxAMatrix*)(double*)&NoneAffineMatrix;
+									GlobalsPerLink[LinkIndex] = Matrix;
+									GlobalLinkFoundFlag = true;
+									break;
+								}
+							}
+						}
+
+						if (!GlobalLinkFoundFlag)
+						{
+							// since now we set use time 0 as ref pose this won't unlikely happen
+							// but leaving it just in case it still has case where it's missing partial bind pose
+							bAnyLinksNotInBindPose = true;
+							LinksWithoutBindPoses += UTF8_TO_TCHAR(Link->GetName());
+							LinksWithoutBindPoses += TEXT("  \n");
+
+							for (int32 ClusterIndex = 0; ClusterIndex < ClusterArray.Num(); ClusterIndex++)
+							{
+								FbxCluster* Cluster = ClusterArray[ClusterIndex];
+								if (Link == Cluster->GetLink())
+								{
+									Cluster->GetTransformLinkMatrix(GlobalsPerLink[LinkIndex]);
+									GlobalLinkFoundFlag = true;
+									break;
+								}
+							}
+						}
+					}
+
+					if (!GlobalLinkFoundFlag)
+					{
+						GlobalsPerLink[LinkIndex] = Link->EvaluateGlobalTransform();
+					}
+
+
+					if (LinkIndex)
+					{
+						LocalsPerLink[LinkIndex] = GlobalsPerLink[ParentIndex].Inverse() * GlobalsPerLink[LinkIndex];
+					}
+					else	// skeleton root
+					{
+						LocalsPerLink[LinkIndex] = GlobalsPerLink[LinkIndex];
+					}
+
+					//Check for nan and for zero scale
+					{
+						bool bFoundNan = false;
+						bool bFoundZeroScale = false;
+						for (int32 i = 0; i < 4; ++i)
+						{
+							if (i < 3)
+							{
+								if (FMath::IsNaN(LocalsPerLink[LinkIndex].GetT()[i]) || FMath::IsNaN(LocalsPerLink[LinkIndex].GetS()[i]))
+								{
+									bFoundNan = true;
+								}
+								if (FMath::IsNearlyZero(LocalsPerLink[LinkIndex].GetS()[i]))
+								{
+									bFoundZeroScale = true;
+								}
+							}
+							if (FMath::IsNaN(LocalsPerLink[LinkIndex].GetQ()[i]))
+							{
+								bFoundNan = true;
+							}
+						}
+
+						//TODO add a user error message
+						ensure(!bFoundNan);
+						ensure(!bFoundZeroScale);
+					}
+				}
+
+
+
+				if (bAnyLinksNotInBindPose)
+				{
+					//TODO error message
+					//
+					//AddTokenizedErrorMessage(FTokenizedMessage::Create(EMessageSeverity::Warning, FText::Format(LOCTEXT("FbxSkeletaLMeshimport_BonesAreMissingFromBindPose", "The following bones are missing from the bind pose:\n{0}\nThis can happen for bones that are not vert weighted. If they are not in the correct orientation after importing,\nplease set the \"Use T0 as ref pose\" option or add them to the bind pose and reimport the skeletal mesh."), FText::FromString(LinksWithoutBindPoses))), FFbxErrors::SkeletalMesh_BonesAreMissingFromBindPose);
+				}
+				return true;
+			}
+
+			FbxNode* FFbxHelper::RecursiveGetFirstMeshNode(FbxNode* Node, FbxNode* NodeToFind)
+			{
+				if (Node == nullptr)
+				{
+					return nullptr;
+				}
+				if (Node->GetMesh() != nullptr)
+					return Node;
+				for (int32 ChildIndex = 0; ChildIndex < Node->GetChildCount(); ++ChildIndex)
+				{
+					FbxNode* MeshNode = RecursiveGetFirstMeshNode(Node->GetChild(ChildIndex), NodeToFind);
+					if (NodeToFind == nullptr)
+					{
+						if (MeshNode != nullptr)
+						{
+							return MeshNode;
+						}
+					}
+					else if (MeshNode == NodeToFind)
+					{
+						return MeshNode;
+					}
+				}
+				return nullptr;
+			}
+
+			FbxNode* FFbxHelper::FindLODGroupNode(FbxNode* NodeLodGroup, int32 LodIndex, FbxNode* NodeToFind)
+			{
+				check(NodeLodGroup->GetChildCount() >= LodIndex);
+				FbxNode* ChildNode = NodeLodGroup->GetChild(LodIndex);
+				if (ChildNode == nullptr)
+				{
+					return nullptr;
+				}
+				return RecursiveGetFirstMeshNode(ChildNode, NodeToFind);
+			}
+
+			void FFbxHelper::RecursiveFindFbxSkelMesh(FbxScene * SDKScene, FbxNode * Node, TArray< TArray<FbxNode*> > & outSkelMeshArray, TArray<FbxNode*> & SkeletonArray)
+			{
+				FbxNode* SkelMeshNode = nullptr;
+				FbxNode* NodeToAdd = Node;
+
+				if (Node->GetMesh() && Node->GetMesh()->GetDeformerCount(FbxDeformer::eSkin) > 0)
+				{
+					SkelMeshNode = Node;
+				}
+				else if (Node->GetNodeAttribute() && Node->GetNodeAttribute()->GetAttributeType() == FbxNodeAttribute::eLODGroup)
+				{
+					// for LODgroup, add the LODgroup to OutSkelMeshArray according to the skeleton that the first child bind to
+					SkelMeshNode = FFbxHelper::FindLODGroupNode(Node, 0, nullptr);
+					// check if the first child is skeletal mesh
+					if (SkelMeshNode != nullptr && !(SkelMeshNode->GetMesh() && SkelMeshNode->GetMesh()->GetDeformerCount(FbxDeformer::eSkin) > 0))
+					{
+						SkelMeshNode = nullptr;
+					}
+					// else NodeToAdd = Node;
+				}
+
+				if (SkelMeshNode)
+				{
+					// find root skeleton
+
+					check(SkelMeshNode->GetMesh() != nullptr);
+					const int32 fbxDeformerCount = SkelMeshNode->GetMesh()->GetDeformerCount();
+					FbxSkin* Deformer = static_cast<FbxSkin*>(SkelMeshNode->GetMesh()->GetDeformer(0, FbxDeformer::eSkin));
+
+					if (Deformer != NULL)
+					{
+						int32 ClusterCount = Deformer->GetClusterCount();
+						bool bFoundCorrectLink = false;
+						for (int32 ClusterId = 0; ClusterId < ClusterCount; ++ClusterId)
+						{
+							FbxNode* RootBoneLink = Deformer->GetCluster(ClusterId)->GetLink(); //Get the bone influences by this first cluster
+							RootBoneLink = GetRootSkeleton(SDKScene, RootBoneLink); // Get the skeleton root itself
+
+							if (RootBoneLink)
+							{
+								bool bAddedToExistingSkeleton = false;
+								for (int32 SkeletonIndex = 0; SkeletonIndex < SkeletonArray.Num(); ++SkeletonIndex)
+								{
+									if (RootBoneLink == SkeletonArray[SkeletonIndex])
+									{
+										// append to existed outSkelMeshArray element
+										TArray<FbxNode*>& TempArray = outSkelMeshArray[SkeletonIndex];
+										TempArray.Add(NodeToAdd);
+										bAddedToExistingSkeleton = true;
+										break;
+									}
+								}
+
+								// if there is no outSkelMeshArray element that is bind to this skeleton
+								// create new element for outSkelMeshArray
+								if (!bAddedToExistingSkeleton)
+								{
+									TArray<FbxNode*>& TempArray = outSkelMeshArray.AddDefaulted_GetRef();
+									TempArray.Add(NodeToAdd);
+									SkeletonArray.Add(RootBoneLink);
+								}
+
+								bFoundCorrectLink = true;
+								break;
+							}
+						}
+
+						// we didn't find the correct link
+						if (!bFoundCorrectLink)
+						{
+							// 								AddTokenizedErrorMessage(
+							// 									FTokenizedMessage::Create(
+							// 									EMessageSeverity::Warning,
+							// 									FText::Format(LOCTEXT("FBX_NoWeightsOnDeformer", "Ignoring mesh {0} because it but no weights."), FText::FromString(UTF8_TO_TCHAR(SkelMeshNode->GetName())))
+							// 								),
+							// 									FFbxErrors::SkeletalMesh_NoWeightsOnDeformer
+							// 								);
+						}
+					}
+				}
+
+				//Skeletalmesh node can have child so let's always iterate trough child
+				{
+					TArray<FbxNode*> ChildScaled;
+					//Sort the node to have the one with no scaling first so we have more chance
+					//to have a root skeletal mesh with no scale. Because scene import do not support
+					//root skeletal mesh containing scale
+					for (int32 ChildIndex = 0; ChildIndex < Node->GetChildCount(); ++ChildIndex)
+					{
+						FbxNode* ChildNode = Node->GetChild(ChildIndex);
+
+						if (!Node->GetNodeAttribute() || Node->GetNodeAttribute()->GetAttributeType() != FbxNodeAttribute::eLODGroup)
+						{
+							FbxVector4 NoScale(1.0, 1.0, 1.0);
+
+							if (ChildNode->EvaluateLocalScaling() == NoScale)
+							{
+								RecursiveFindFbxSkelMesh(SDKScene, ChildNode, outSkelMeshArray, SkeletonArray);
+							}
+							else
+							{
+								ChildScaled.Add(ChildNode);
+							}
+						}
+					}
+
+					for (FbxNode* ChildNode : ChildScaled)
+					{
+						RecursiveFindFbxSkelMesh(SDKScene, ChildNode, outSkelMeshArray, SkeletonArray);
+					}
+				}
+			}
+
+			void FFbxHelper::RecursiveFindRigidMesh(FbxScene * SDKScene, FbxNode * Node, TArray< TArray<FbxNode*> > & outSkelMeshArray, TArray<FbxNode*> & SkeletonArray)
+			{
+				bool bRigidNodeFound = false;
+				FbxNode* RigidMeshNode = nullptr;
+
+				if (Node->GetMesh())
+				{
+					// ignore skeletal mesh
+					if (Node->GetMesh()->GetDeformerCount(FbxDeformer::eSkin) == 0)
+					{
+						RigidMeshNode = Node;
+						bRigidNodeFound = true;
+					}
+				}
+				else if (Node->GetNodeAttribute() && Node->GetNodeAttribute()->GetAttributeType() == FbxNodeAttribute::eLODGroup)
+				{
+					// for LODgroup, add the LODgroup to OutSkelMeshArray according to the skeleton that the first child bind to
+					FbxNode* FirstLOD = FindLODGroupNode(Node, 0, nullptr);
+					// check if the first child is skeletal mesh
+					if (FirstLOD != nullptr && FirstLOD->GetMesh())
+					{
+						if (FirstLOD->GetMesh()->GetDeformerCount(FbxDeformer::eSkin) == 0)
+						{
+							bRigidNodeFound = true;
+						}
+					}
+
+					if (bRigidNodeFound)
+					{
+						RigidMeshNode = Node;
+					}
+				}
+
+				if (bRigidNodeFound)
+				{
+					// find root skeleton
+					FbxNode* Link = GetRootSkeleton(SDKScene, RigidMeshNode);
+
+					int32 i;
+					for (i = 0; i < SkeletonArray.Num(); i++)
+					{
+						if (Link == SkeletonArray[i])
+						{
+							// append to existed outSkelMeshArray element
+							TArray<FbxNode*>& TempArray = outSkelMeshArray[i];
+							TempArray.Add(RigidMeshNode);
+							break;
+						}
+					}
+
+					// if there is no outSkelMeshArray element that is bind to this skeleton
+					// create new element for outSkelMeshArray
+					if (i == SkeletonArray.Num())
+					{
+						TArray<FbxNode*>& TempArray = outSkelMeshArray.AddDefaulted_GetRef();
+						TempArray.Add(RigidMeshNode);
+						SkeletonArray.Add(Link);
+					}
+				}
+
+				// for LODGroup, we will not deep in.
+				if (!(Node->GetNodeAttribute() && Node->GetNodeAttribute()->GetAttributeType() == FbxNodeAttribute::eLODGroup))
+				{
+					int32 ChildIndex;
+					for (ChildIndex = 0; ChildIndex < Node->GetChildCount(); ++ChildIndex)
+					{
+						RecursiveFindRigidMesh(SDKScene, Node->GetChild(ChildIndex), outSkelMeshArray, SkeletonArray);
+					}
+				}
+			}
+
+			void FFbxHelper::RecursiveFixSkeleton(FbxScene * SDKScene, FbxNode * Node, TArray<FbxNode*> & SkelMeshes, bool bImportNestedMeshes)
+			{
+				FbxNodeAttribute* Attr = Node->GetNodeAttribute();
+				bool NodeIsLodGroup = (Attr && (Attr->GetAttributeType() == FbxNodeAttribute::eLODGroup));
+				if (!NodeIsLodGroup)
+				{
+					for (int32 i = 0; i < Node->GetChildCount(); i++)
+					{
+						RecursiveFixSkeleton(SDKScene, Node->GetChild(i), SkelMeshes, bImportNestedMeshes);
+					}
+				}
+
+				if (Attr && (Attr->GetAttributeType() == FbxNodeAttribute::eMesh || Attr->GetAttributeType() == FbxNodeAttribute::eNull))
+				{
+					if (bImportNestedMeshes && Attr->GetAttributeType() == FbxNodeAttribute::eMesh)
+					{
+						// for leaf mesh, keep them as mesh
+						int32 ChildCount = Node->GetChildCount();
+						int32 ChildIndex;
+						for (ChildIndex = 0; ChildIndex < ChildCount; ChildIndex++)
+						{
+							FbxNode* Child = Node->GetChild(ChildIndex);
+							if (Child->GetMesh() == NULL)
+							{
+								break;
+							}
+						}
+
+						if (ChildIndex != ChildCount)
+						{
+							// Remove from the mesh list it is no longer a mesh
+							SkelMeshes.Remove(Node);
+
+							//replace with skeleton
+							FbxSkeleton* lSkeleton = FbxSkeleton::Create(SDKScene->GetFbxManager(), "");
+							Node->SetNodeAttribute(lSkeleton);
+							lSkeleton->SetSkeletonType(FbxSkeleton::eLimbNode);
+						}
+						else // this mesh may be not in skeleton mesh list. If not, add it.
+						{
+							if (!SkelMeshes.Contains(Node))
+							{
+								SkelMeshes.Add(Node);
+							}
+						}
+					}
+					else
+					{
+						// Remove from the mesh list it is no longer a mesh
+						SkelMeshes.Remove(Node);
+
+						//replace with skeleton
+						FbxSkeleton* lSkeleton = FbxSkeleton::Create(SDKScene->GetFbxManager(), "");
+						Node->SetNodeAttribute(lSkeleton);
+						lSkeleton->SetSkeletonType(FbxSkeleton::eLimbNode);
+					}
+				}
+			}
+
+		}//ns Private
+	}//ns Interchange
+}//ns UE
diff --git a/Engine/Plugins/Runtime/Interchange/Source/Parsers/Fbx/Private/FbxHelper.h b/Engine/Plugins/Runtime/Interchange/Source/Parsers/Fbx/Private/FbxHelper.h
new file mode 100644
--- /dev/null
+++ b/Engine/Plugins/Runtime/Interchange/Source/Parsers/Fbx/Private/FbxHelper.h
@@ -0,0 +1,63 @@
+// Copyright Epic Games, Inc. All Rights Reserved.
+
+#pragma once
+
+#include "CoreMinimal.h"
+#include "FbxInclude.h"
+
+namespace UE
+{
+	namespace Interchange
+	{
+		namespace Private
+		{
+			class FPayloadContextBase
+			{
+			public:
+				virtual ~FPayloadContextBase() {}
+				virtual FString GetPayloadType() const { return FString(); }
+				virtual bool FetchPayloadToFile(const FString& PayloadFilepath, TArray<FString>& JSonErrorMessages) { return false; }
+			};
+
+			struct FFbxHelper
+			{
+			public:
+
+				/**
+				 * Return the name of an FbxObject, return empty string if the object is null.
+				 */
+				static FString GetFbxObjectName(const FbxObject* Object);
+
+				/**
+				 * Return a string with the name of all the parent in the hierarchy separate by a dot( . ) from the fbx root node to the specified node.
+				 * This is a way to have a valid unique ID for a fbx node that will be the same if the fbx change when we re-import.
+				 * Using the fbx sdk int32 uniqueID is not valid anymore if the fbx is re-exported.
+				 */
+				static FString GetFbxNodeHierarchyName(const FbxNode* Node);
+
+				//Skeletalmesh helper
+				static void FindSkeletalMeshes(FbxScene* SDKScene, TArray< TArray<FbxNode*> >& outSkelMeshArray, bool bCombineSkeletalMesh, bool bForceFindRigid);
+				static void FindAllLODGroupNode(TArray<FbxNode*>& OutNodeInLod, FbxNode* NodeLodGroup, int32 LodIndex);
+				static bool FindSkeletonJoints(FbxScene* SDKScene, TArray<FbxNode*>& NodeArray, TArray<FbxNode*>& SortedLinks, FbxArray<FbxAMatrix>& LocalsPerLink);
+				static FbxNode* FindLODGroupNode(FbxNode* NodeLodGroup, int32 LodIndex, FbxNode* NodeToFind);
+
+			private:
+
+				/* Skletalmesh private helper begin */
+
+				static void RecursiveGetAllMeshNode(TArray<FbxNode*>& OutAllNode, FbxNode* Node);
+				static bool IsUnrealBone(FbxNode* Link);
+				static void RecursiveBuildSkeleton(FbxNode* Link, TArray<FbxNode*>& OutSortedLinks);
+				static bool RetrievePoseFromBindPose(FbxScene* SDKScene, const TArray<FbxNode*>& NodeArray, FbxArray<FbxPose*>& PoseArray);
+				static FbxNode* GetRootSkeleton(FbxScene* SDKScene, FbxNode* Link);
+				static void BuildSkeletonSystem(FbxScene* SDKScene, TArray<FbxCluster*>& ClusterArray, TArray<FbxNode*>& OutSortedLinks);
+				static FbxNode* RecursiveGetFirstMeshNode(FbxNode* Node, FbxNode* NodeToFind);
+				static void RecursiveFindFbxSkelMesh(FbxScene* SDKScene, FbxNode* Node, TArray< TArray<FbxNode*> >& outSkelMeshArray, TArray<FbxNode*>& SkeletonArray);
+				static void RecursiveFindRigidMesh(FbxScene* SDKScene, FbxNode* Node, TArray< TArray<FbxNode*> >& outSkelMeshArray, TArray<FbxNode*>& SkeletonArray);
+				static void RecursiveFixSkeleton(FbxScene* SDKScene, FbxNode* Node, TArray<FbxNode*>& SkelMeshes, bool bImportNestedMeshes);
+
+				/* Skletalmesh private helper end */
+			};
+		}//ns Private
+	}//ns Interchange
+}//ns UE
diff --git a/Engine/Plugins/Runtime/Interchange/Source/Parsers/Fbx/Private/FbxMaterial.cpp b/Engine/Plugins/Runtime/Interchange/Source/Parsers/Fbx/Private/FbxMaterial.cpp
--- a/Engine/Plugins/Runtime/Interchange/Source/Parsers/Fbx/Private/FbxMaterial.cpp
+++ b/Engine/Plugins/Runtime/Interchange/Source/Parsers/Fbx/Private/FbxMaterial.cpp
@@ -4,6 +4,7 @@
 
 #include "CoreMinimal.h"
 #include "FbxConvert.h"
+#include "FbxHelper.h"
 #include "FbxInclude.h"
 #include "InterchangeMaterialNode.h"
 #include "InterchangeTextureNode.h"
@@ -17,8 +18,10 @@ namespace UE
 		{
 			UInterchangeMaterialNode* FFbxMaterial::CreateMaterialNode(UInterchangeBaseNodeContainer& NodeContainer, const FString& NodeName, TArray<FString>& JSonErrorMessages)
 			{
-				FName DisplayLabel = *NodeName;
 				FName NodeUID(*NodeName);
+				FString MaterialNameNoSkin = UE::Interchange::Material::RemoveSkinFromName(NodeName);
+
+				FName DisplayLabel = *MaterialNameNoSkin;
 				UInterchangeMaterialNode* MaterialNode = NewObject<UInterchangeMaterialNode>(&NodeContainer, NAME_None);
 				if (!ensure(MaterialNode))
 				{
@@ -52,7 +55,7 @@ namespace UE
 			UInterchangeMaterialNode* FFbxMaterial::AddNodeMaterial(FbxSurfaceMaterial* SurfaceMaterial, UInterchangeBaseNodeContainer& NodeContainer, TArray<FString>& JSonErrorMessages)
 			{
 				//Create a material node
-				FString MaterialName = FFbxConvert::MakeString(FFbxConvert::MakeName(SurfaceMaterial->GetName()));
+				FString MaterialName = FFbxHelper::GetFbxObjectName(SurfaceMaterial);
 				FName NodeUID(*MaterialName);
 				UInterchangeMaterialNode* MaterialNode = Cast<UInterchangeMaterialNode>(NodeContainer.GetNode(NodeUID));
 				if (!MaterialNode)
@@ -83,6 +86,11 @@ namespace UE
 								{
 									FbxFileTexture* FbxTextureFilePath = FbxProperty.GetSrcObject<FbxFileTexture>(TextureIndex);
 									FString TextureFilename = UTF8_TO_TCHAR(FbxTextureFilePath->GetFileName());
+									//Only import texture that exist on disk
+									if (!FPaths::FileExists(TextureFilename))
+									{
+										continue;
+									}
 									//Create a texture node and make it child of the material node
 									TArray<FString> JsonErrorMessage;
 									FName NodeUID(*TextureFilename);
diff --git a/Engine/Plugins/Runtime/Interchange/Source/Parsers/Fbx/Private/FbxMesh.cpp b/Engine/Plugins/Runtime/Interchange/Source/Parsers/Fbx/Private/FbxMesh.cpp
new file mode 100644
--- /dev/null
+++ b/Engine/Plugins/Runtime/Interchange/Source/Parsers/Fbx/Private/FbxMesh.cpp
@@ -0,0 +1,1146 @@
+// Copyright Epic Games, Inc. All Rights Reserved.
+
+#include "FbxMesh.h"
+
+#include "CoreMinimal.h"
+#include "FbxConvert.h"
+#include "FbxHelper.h"
+#include "FbxInclude.h"
+#include "FbxSkeletalMesh.h"
+#include "MeshDescription.h"
+#include "SkeletalMeshAttributes.h"
+#include "StaticMeshAttributes.h"
+
+namespace UE
+{
+	namespace Interchange
+	{
+		namespace Private
+		{
+			// Wraps some common code useful for multiple fbx import code path
+			struct FFBXUVs
+			{
+				// constructor
+				FFBXUVs(FbxMesh* Mesh)
+					: UniqueUVCount(0)
+				{
+					check(Mesh);
+
+					//
+					//	store the UVs in arrays for fast access in the later looping of triangles 
+					//
+					// mapping from UVSets to Fbx LayerElementUV
+					// Fbx UVSets may be duplicated, remove the duplicated UVSets in the mapping 
+					int32 LayerCount = Mesh->GetLayerCount();
+					if (LayerCount > 0)
+					{
+						int32 UVLayerIndex;
+						for (UVLayerIndex = 0; UVLayerIndex < LayerCount; UVLayerIndex++)
+						{
+							FbxLayer* lLayer = Mesh->GetLayer(UVLayerIndex);
+							int UVSetCount = lLayer->GetUVSetCount();
+							if (UVSetCount)
+							{
+								FbxArray<FbxLayerElementUV const*> EleUVs = lLayer->GetUVSets();
+								for (int UVIndex = 0; UVIndex < UVSetCount; UVIndex++)
+								{
+									FbxLayerElementUV const* ElementUV = EleUVs[UVIndex];
+									if (ElementUV)
+									{
+										const char* UVSetName = ElementUV->GetName();
+										FString LocalUVSetName = UTF8_TO_TCHAR(UVSetName);
+										if (LocalUVSetName.IsEmpty())
+										{
+											LocalUVSetName = TEXT("UVmap_") + FString::FromInt(UVLayerIndex);
+										}
+
+										UVSets.AddUnique(LocalUVSetName);
+									}
+								}
+							}
+						}
+					}
+
+
+					// If the the UV sets are named using the following format (UVChannel_X; where X ranges from 1 to 4)
+					// we will re-order them based on these names.  Any UV sets that do not follow this naming convention
+					// will be slotted into available spaces.
+					if (UVSets.Num())
+					{
+						for (int32 ChannelNumIdx = 0; ChannelNumIdx < 4; ChannelNumIdx++)
+						{
+							FString ChannelName = FString::Printf(TEXT("UVChannel_%d"), ChannelNumIdx + 1);
+							int32 SetIdx = UVSets.Find(ChannelName);
+
+							// If the specially formatted UVSet name appears in the list and it is in the wrong spot,
+							// we will swap it into the correct spot.
+							if (SetIdx != INDEX_NONE && SetIdx != ChannelNumIdx)
+							{
+								// If we are going to swap to a position that is outside the bounds of the
+								// array, then we pad out to that spot with empty data.
+								for (int32 ArrSize = UVSets.Num(); ArrSize < ChannelNumIdx + 1; ArrSize++)
+								{
+									UVSets.Add(FString(TEXT("")));
+								}
+								//Swap the entry into the appropriate spot.
+								UVSets.Swap(SetIdx, ChannelNumIdx);
+							}
+						}
+					}
+				}
+
+				void Phase2(FbxMesh* Mesh)
+				{
+					//
+					//	store the UVs in arrays for fast access in the later looping of triangles 
+					//
+					UniqueUVCount = UVSets.Num();
+					if (UniqueUVCount > 0)
+					{
+						LayerElementUV.AddZeroed(UniqueUVCount);
+						UVReferenceMode.AddZeroed(UniqueUVCount);
+						UVMappingMode.AddZeroed(UniqueUVCount);
+					}
+					for (int32 UVIndex = 0; UVIndex < UniqueUVCount; UVIndex++)
+					{
+						LayerElementUV[UVIndex] = NULL;
+						for (int32 UVLayerIndex = 0, LayerCount = Mesh->GetLayerCount(); UVLayerIndex < LayerCount; UVLayerIndex++)
+						{
+							FbxLayer* lLayer = Mesh->GetLayer(UVLayerIndex);
+							int UVSetCount = lLayer->GetUVSetCount();
+							if (UVSetCount)
+							{
+								FbxArray<FbxLayerElementUV const*> EleUVs = lLayer->GetUVSets();
+								for (int32 FbxUVIndex = 0; FbxUVIndex < UVSetCount; FbxUVIndex++)
+								{
+									FbxLayerElementUV const* ElementUV = EleUVs[FbxUVIndex];
+									if (ElementUV)
+									{
+										const char* UVSetName = ElementUV->GetName();
+										FString LocalUVSetName = UTF8_TO_TCHAR(UVSetName);
+										if (LocalUVSetName.IsEmpty())
+										{
+											LocalUVSetName = TEXT("UVmap_") + FString::FromInt(UVLayerIndex);
+										}
+										if (LocalUVSetName == UVSets[UVIndex])
+										{
+											LayerElementUV[UVIndex] = ElementUV;
+											UVReferenceMode[UVIndex] = ElementUV->GetReferenceMode();
+											UVMappingMode[UVIndex] = ElementUV->GetMappingMode();
+											break;
+										}
+									}
+								}
+							}
+						}
+					}
+
+					if (UniqueUVCount > MAX_MESH_TEXTURE_COORDS_MD)
+					{
+						//TODO Error Message
+						//FbxImporter->AddTokenizedErrorMessage(FTokenizedMessage::Create(EMessageSeverity::Warning, FText::Format(LOCTEXT("Error_TooMuchUVChannel", "Reached the maximum number of UV Channels for a Static Mesh({0}) - discarding {1} UV Channels"), FText::AsNumber(MAX_MESH_TEXTURE_COORDS_MD), FText::AsNumber(UniqueUVCount - MAX_MESH_TEXTURE_COORDS_MD))), FFbxErrors::Generic_Mesh_TooMuchUVChannels);
+					}
+
+					UniqueUVCount = FMath::Min<int32>(UniqueUVCount, MAX_MESH_TEXTURE_COORDS_MD);
+				}
+
+				int32 FindLightUVIndex() const
+				{
+					// See if any of our UV set entry names match LightMapUV.
+					for (int32 UVSetIdx = 0; UVSetIdx < UVSets.Num(); UVSetIdx++)
+					{
+						if (UVSets[UVSetIdx] == TEXT("LightMapUV"))
+						{
+							return UVSetIdx;
+						}
+					}
+
+					// not found
+					return INDEX_NONE;
+				}
+
+				// @param FaceCornerIndex usually TriangleIndex * 3 + CornerIndex but more complicated for mixed n-gons
+				int32 ComputeUVIndex(int32 UVLayerIndex, int32 lControlPointIndex, int32 FaceCornerIndex) const
+				{
+					int32 UVMapIndex = (UVMappingMode[UVLayerIndex] == FbxLayerElement::eByControlPoint) ? lControlPointIndex : FaceCornerIndex;
+
+					int32 Ret;
+
+					if (UVReferenceMode[UVLayerIndex] == FbxLayerElement::eDirect)
+					{
+						Ret = UVMapIndex;
+					}
+					else
+					{
+						FbxLayerElementArrayTemplate<int>& Array = LayerElementUV[UVLayerIndex]->GetIndexArray();
+						Ret = Array.GetAt(UVMapIndex);
+					}
+
+					return Ret;
+				}
+
+				// todo: is that needed? could the dtor do it?
+				void Cleanup()
+				{
+					//
+					// clean up.  This needs to happen before the mesh is destroyed
+					//
+					LayerElementUV.Empty();
+					UVReferenceMode.Empty();
+					UVMappingMode.Empty();
+				}
+
+				TArray<FString> UVSets;
+				TArray<FbxLayerElementUV const*> LayerElementUV;
+				TArray<FbxLayerElement::EReferenceMode> UVReferenceMode;
+				TArray<FbxLayerElement::EMappingMode> UVMappingMode;
+				int32 UniqueUVCount;
+			};
+
+			FMeshDescriptionImporter::FMeshDescriptionImporter(FMeshDescription* InMeshDescription, FbxNode* InMeshNode, FbxScene* InSDKScene, FbxGeometryConverter* InSDKGeometryConverter)
+				: MeshDescription(InMeshDescription)
+				, MeshNode(InMeshNode)
+				, SDKScene(InSDKScene)
+				, SDKGeometryConverter(InSDKGeometryConverter)
+			{
+				bInitialized = ensure(MeshDescription) && ensure(MeshNode) && ensure(SDKScene) && ensure(SDKGeometryConverter);
+			}
+			
+			bool FMeshDescriptionImporter::FillStaticMeshDescriptionFromFbxMesh()
+			{
+				return FillMeshDescriptionFromFbxMesh(EMeshType::Static, nullptr);
+			}
+
+			bool FMeshDescriptionImporter::FillSkinnedMeshDescriptionFromFbxMesh(TArray<FbxNode*>* SortedJoints)
+			{
+				FbxMesh* Mesh = MeshNode->GetMesh();
+				if (!ensure(Mesh))
+				{
+					//TODO log an error
+					return false;
+				}
+				FbxSkin* Skin = (FbxSkin*)Mesh->GetDeformer(0, FbxDeformer::eSkin);
+				EMeshType MeshType = Skin != nullptr ? EMeshType::Skinned : EMeshType::Rigid;
+				return FillMeshDescriptionFromFbxMesh(MeshType, SortedJoints);
+			}
+
+			bool FMeshDescriptionImporter::FillMeshDescriptionFromFbxMesh(EMeshType MeshType, TArray<FbxNode*>* SortedJoints)
+			{
+				if (!ensure(bInitialized) || !ensure(MeshDescription))
+				{
+					return false;
+				}
+
+				FbxMesh* Mesh = MeshNode->GetMesh();
+				FStaticMeshAttributes Attributes(*MeshDescription);
+
+				//Get the base layer of the mesh
+				FbxLayer* BaseLayer = Mesh->GetLayer(0);
+				if (BaseLayer == NULL)
+				{
+					//TODO add error message
+					//AddTokenizedErrorMessage(FTokenizedMessage::Create(EMessageSeverity::Error, FText::Format(LOCTEXT("Error_NoGeometryInMesh", "There is no geometry information in mesh '{0}'"), FText::FromString(Mesh->GetName()))), FFbxErrors::Generic_Mesh_NoGeometry);
+					return false;
+				}
+
+				FFBXUVs FBXUVs(Mesh);
+
+				if (MeshType == EMeshType::Static)
+				{
+// 					int32 FBXNamedLightMapCoordinateIndex = FBXUVs.FindLightUVIndex();
+// 					if (FBXNamedLightMapCoordinateIndex != INDEX_NONE)
+// 					{
+// 						StaticMesh->SetLightMapCoordinateIndex(FBXNamedLightMapCoordinateIndex);
+// 					}
+				}
+
+				//
+				// create materials
+				//
+				//Create a material name array in the node order, also fill the Meshdescription PolygonGroup
+				TArray<FName> MaterialNames;
+				const int32 MaterialCount = MeshNode->GetMaterialCount();
+				MaterialNames.Reserve(MaterialCount);
+				for (int32 MaterialIndex = 0; MaterialIndex < MaterialCount; ++MaterialIndex)
+				{
+					FbxSurfaceMaterial* FbxMaterial = MeshNode->GetMaterial(MaterialIndex);
+					MaterialNames.Add(*FFbxHelper::GetFbxObjectName(FbxMaterial));
+				}
+
+				// Must do this before triangulating the mesh due to an FBX bug in TriangulateMeshAdvance
+				int32 LayerSmoothingCount = Mesh->GetLayerCount(FbxLayerElement::eSmoothing);
+				for (int32 i = 0; i < LayerSmoothingCount; i++)
+				{
+					FbxLayerElementSmoothing const* SmoothingInfo = Mesh->GetLayer(0)->GetSmoothing();
+					if (SmoothingInfo && SmoothingInfo->GetMappingMode() != FbxLayerElement::eByPolygon)
+					{
+						SDKGeometryConverter->ComputePolygonSmoothingFromEdgeSmoothing(Mesh, i);
+					}
+				}
+
+				if (!Mesh->IsTriangleMesh())
+				{
+					const bool bReplace = true;
+					FbxNodeAttribute* ConvertedNode = SDKGeometryConverter->Triangulate(Mesh, bReplace);
+
+					if (ConvertedNode != NULL && ConvertedNode->GetAttributeType() == FbxNodeAttribute::eMesh)
+					{
+						Mesh = (fbxsdk::FbxMesh*)ConvertedNode;
+					}
+					else
+					{
+						//TODO add an error message
+						//AddTokenizedErrorMessage(FTokenizedMessage::Create(EMessageSeverity::Warning, FText::Format(LOCTEXT("Error_FailedToTriangulate", "Unable to triangulate mesh '{0}'"), FText::FromString(Mesh->GetName()))), FFbxErrors::Generic_Mesh_TriangulationFailed);
+						return false; // not clean, missing some dealloc
+					}
+				}
+
+				// renew the base layer
+				BaseLayer = Mesh->GetLayer(0);
+
+				//
+				// get the smoothing group layer
+				//
+				bool bSmoothingAvailable = false;
+
+				FbxLayerElementSmoothing* SmoothingInfo = BaseLayer->GetSmoothing();
+				FbxLayerElement::EReferenceMode SmoothingReferenceMode(FbxLayerElement::eDirect);
+				FbxLayerElement::EMappingMode SmoothingMappingMode(FbxLayerElement::eByEdge);
+				if (SmoothingInfo)
+				{
+					if (SmoothingInfo->GetMappingMode() == FbxLayerElement::eByPolygon)
+					{
+						//Convert the base layer to edge smoothing
+						SDKGeometryConverter->ComputeEdgeSmoothingFromPolygonSmoothing(Mesh, 0);
+						BaseLayer = Mesh->GetLayer(0);
+						SmoothingInfo = BaseLayer->GetSmoothing();
+					}
+
+					if (SmoothingInfo->GetMappingMode() == FbxLayerElement::eByEdge)
+					{
+						bSmoothingAvailable = true;
+					}
+
+					SmoothingReferenceMode = SmoothingInfo->GetReferenceMode();
+					SmoothingMappingMode = SmoothingInfo->GetMappingMode();
+
+					//If all smooth group are 0's, we can say that there is no valid smoothing data, but only for skeletal mesh
+					//Static mesh will be faceted in this case since we will not compute smooth group from normals like we do here
+					//This is legacy code to get the same result has before
+					if (MeshType == EMeshType::Skinned || MeshType == EMeshType::Rigid)
+					{
+						// Check and see if the smooothing data is valid.  If not generate it from the normals
+						BaseLayer = Mesh->GetLayer(0);
+						if (BaseLayer)
+						{
+							const FbxLayerElementSmoothing* SmoothingLayer = BaseLayer->GetSmoothing();
+
+							if (SmoothingLayer)
+							{
+								bool bValidSmoothingData = false;
+								FbxLayerElementArrayTemplate<int32>& Array = SmoothingLayer->GetDirectArray();
+								for (int32 SmoothingIndex = 0; SmoothingIndex < Array.GetCount(); ++SmoothingIndex)
+								{
+									if (Array[SmoothingIndex] != 0)
+									{
+										bValidSmoothingData = true;
+										break;
+									}
+								}
+
+								if (!bValidSmoothingData && Mesh->GetPolygonVertexCount() > 0)
+								{
+									SDKGeometryConverter->ComputeEdgeSmoothingFromNormals(Mesh);
+									BaseLayer = Mesh->GetLayer(0);
+									SmoothingInfo = BaseLayer->GetSmoothing();
+								}
+							}
+						}
+					}
+				}
+
+				//
+				//	get the "material index" layer.  Do this AFTER the triangulation step as that may reorder material indices
+				//
+				FbxLayerElementMaterial* LayerElementMaterial = BaseLayer->GetMaterials();
+				FbxLayerElement::EMappingMode MaterialMappingMode = LayerElementMaterial ?
+					LayerElementMaterial->GetMappingMode() : FbxLayerElement::eByPolygon;
+
+				//	todo second phase UV, ok to put in first phase?
+				FBXUVs.Phase2(Mesh);
+
+				//
+				// get the first vertex color layer
+				//
+				FbxLayerElementVertexColor* LayerElementVertexColor = BaseLayer->GetVertexColors();
+				FbxLayerElement::EReferenceMode VertexColorReferenceMode(FbxLayerElement::eDirect);
+				FbxLayerElement::EMappingMode VertexColorMappingMode(FbxLayerElement::eByControlPoint);
+				if (LayerElementVertexColor)
+				{
+					VertexColorReferenceMode = LayerElementVertexColor->GetReferenceMode();
+					VertexColorMappingMode = LayerElementVertexColor->GetMappingMode();
+				}
+
+				//
+				// get the first normal layer
+				//
+				FbxLayerElementNormal* LayerElementNormal = BaseLayer->GetNormals();
+				FbxLayerElementTangent* LayerElementTangent = BaseLayer->GetTangents();
+				FbxLayerElementBinormal* LayerElementBinormal = BaseLayer->GetBinormals();
+
+				//whether there is normal, tangent and binormal data in this mesh
+				bool bHasNTBInformation = LayerElementNormal && LayerElementTangent && LayerElementBinormal;
+
+				FbxLayerElement::EReferenceMode NormalReferenceMode(FbxLayerElement::eDirect);
+				FbxLayerElement::EMappingMode NormalMappingMode(FbxLayerElement::eByControlPoint);
+				if (LayerElementNormal)
+				{
+					NormalReferenceMode = LayerElementNormal->GetReferenceMode();
+					NormalMappingMode = LayerElementNormal->GetMappingMode();
+				}
+
+				FbxLayerElement::EReferenceMode TangentReferenceMode(FbxLayerElement::eDirect);
+				FbxLayerElement::EMappingMode TangentMappingMode(FbxLayerElement::eByControlPoint);
+				if (LayerElementTangent)
+				{
+					TangentReferenceMode = LayerElementTangent->GetReferenceMode();
+					TangentMappingMode = LayerElementTangent->GetMappingMode();
+				}
+
+				FbxLayerElement::EReferenceMode BinormalReferenceMode(FbxLayerElement::eDirect);
+				FbxLayerElement::EMappingMode BinormalMappingMode(FbxLayerElement::eByControlPoint);
+				if (LayerElementBinormal)
+				{
+					BinormalReferenceMode = LayerElementBinormal->GetReferenceMode();
+					BinormalMappingMode = LayerElementBinormal->GetMappingMode();
+				}
+
+				bool bHasNonDegeneratePolygons = false;
+				{
+					MeshDescription->SuspendVertexInstanceIndexing();
+					MeshDescription->SuspendEdgeIndexing();
+					MeshDescription->SuspendPolygonIndexing();
+					MeshDescription->SuspendPolygonGroupIndexing();
+					MeshDescription->SuspendUVIndexing();
+
+					TRACE_CPUPROFILER_EVENT_SCOPE(BuildTriangles);
+
+					// Construct the matrices for the conversion from right handed to left handed system
+					FbxAMatrix TotalMatrix;
+					FbxAMatrix TotalMatrixForNormal;
+					TotalMatrix = ComputeNodeMatrix(MeshNode);
+					TotalMatrixForNormal = TotalMatrix.Inverse();
+					TotalMatrixForNormal = TotalMatrixForNormal.Transpose();
+					int32 PolygonCount = Mesh->GetPolygonCount();
+
+					if (PolygonCount == 0)
+					{
+						//TODO add an error log
+						//AddTokenizedErrorMessage(FTokenizedMessage::Create(EMessageSeverity::Error, FText::Format(LOCTEXT("Error_NoPolygonFoundInMesh", "No polygon were found on mesh  '{0}'"), FText::FromString(Mesh->GetName()))), FFbxErrors::StaticMesh_NoTriangles);
+						return false;
+					}
+
+					int32 VertexCount = Mesh->GetControlPointsCount();
+					bool OddNegativeScale = IsOddNegativeScale(TotalMatrix);
+
+					TVertexAttributesRef<FVector> VertexPositions = Attributes.GetVertexPositions();
+					TVertexInstanceAttributesRef<FVector> VertexInstanceNormals = Attributes.GetVertexInstanceNormals();
+					TVertexInstanceAttributesRef<FVector> VertexInstanceTangents = Attributes.GetVertexInstanceTangents();
+					TVertexInstanceAttributesRef<float> VertexInstanceBinormalSigns = Attributes.GetVertexInstanceBinormalSigns();
+					TVertexInstanceAttributesRef<FVector4> VertexInstanceColors = Attributes.GetVertexInstanceColors();
+					TVertexInstanceAttributesRef<FVector2D> VertexInstanceUVs = Attributes.GetVertexInstanceUVs();
+					TEdgeAttributesRef<bool> EdgeHardnesses = Attributes.GetEdgeHardnesses();
+					TPolygonGroupAttributesRef<FName> PolygonGroupImportedMaterialSlotNames = Attributes.GetPolygonGroupMaterialSlotNames();
+
+					int32 VertexOffset = MeshDescription->Vertices().Num();
+					int32 VertexInstanceOffset = MeshDescription->VertexInstances().Num();
+					int32 PolygonOffset = MeshDescription->Polygons().Num();
+
+					// The below code expects Num() to be equivalent to GetArraySize(), i.e. that all added elements are appended, not inserted into existing gaps
+					check(VertexOffset == MeshDescription->Vertices().GetArraySize());
+					check(VertexInstanceOffset == MeshDescription->VertexInstances().GetArraySize());
+					check(PolygonOffset == MeshDescription->Polygons().GetArraySize());
+
+					TMap<int32, FPolygonGroupID> PolygonGroupMapping;
+
+					// When importing multiple mesh pieces to the same static mesh.  Ensure each mesh piece has the same number of Uv's
+					int32 ExistingUVCount = VertexInstanceUVs.GetNumChannels();
+
+					int32 NumUVs = FMath::Max(FBXUVs.UniqueUVCount, ExistingUVCount);
+					NumUVs = FMath::Min<int32>(MAX_MESH_TEXTURE_COORDS_MD, NumUVs);
+					// At least one UV set must exist.  
+					// @todo: this needn't be true; we should be able to handle zero UV channels
+					NumUVs = FMath::Max(1, NumUVs);
+
+					//Make sure all Vertex instance have the correct number of UVs
+					VertexInstanceUVs.SetNumChannels(NumUVs);
+					MeshDescription->SetNumUVChannels(NumUVs);
+
+					TArray<int32> UVOffsets;
+					UVOffsets.SetNumUninitialized(NumUVs);
+					for (int UVChannel = 0; UVChannel < NumUVs; UVChannel++)
+					{
+						UVOffsets[UVChannel] = MeshDescription->UVs(UVChannel).GetArraySize();
+					}
+
+					//Fill the vertex array
+					MeshDescription->ReserveNewVertices(VertexCount);
+					for (int32 VertexIndex = 0; VertexIndex < VertexCount; ++VertexIndex)
+					{
+						int32 RealVertexIndex = VertexOffset + VertexIndex;
+						FbxVector4 FbxPosition = Mesh->GetControlPoints()[VertexIndex];
+						FbxPosition = TotalMatrix.MultT(FbxPosition);
+						const FVector VertexPosition = FFbxConvert::ConvertPos(FbxPosition);
+
+						FVertexID AddedVertexId = MeshDescription->CreateVertex();
+						VertexPositions[AddedVertexId] = VertexPosition;
+						if (AddedVertexId.GetValue() != RealVertexIndex)
+						{
+							//TODO add an error log
+							//AddTokenizedErrorMessage(FTokenizedMessage::Create(EMessageSeverity::Error, FText::Format(LOCTEXT("Error_CannotCreateVertex", "Cannot create valid vertex for mesh '{0}'"), FText::FromString(Mesh->GetName()))), FFbxErrors::StaticMesh_BuildError);
+							return false;
+						}
+					}
+
+					// Fill the UV arrays
+					for (int32 UVLayerIndex = 0; UVLayerIndex < FBXUVs.UniqueUVCount; UVLayerIndex++)
+					{
+						check(FBXUVs.LayerElementUV[UVLayerIndex]);
+						if (FBXUVs.LayerElementUV[UVLayerIndex] != nullptr)
+						{
+							int32 UVCount = FBXUVs.LayerElementUV[UVLayerIndex]->GetDirectArray().GetCount();
+							TUVAttributesRef<FVector2D> UVCoordinates = MeshDescription->UVAttributes(UVLayerIndex).GetAttributesRef<FVector2D>(MeshAttribute::UV::UVCoordinate);
+							MeshDescription->ReserveNewUVs(UVCount, UVLayerIndex);
+							for (int32 UVIndex = 0; UVIndex < UVCount; UVIndex++)
+							{
+								FUVID UVID = MeshDescription->CreateUV(UVLayerIndex);
+								FbxVector2 UVVector = FBXUVs.LayerElementUV[UVLayerIndex]->GetDirectArray().GetAt(UVIndex);
+								UVCoordinates[UVID] = FVector2D(static_cast<float>(UVVector[0]), 1.0f - static_cast<float>(UVVector[1]));	// flip the Y of UVs for DirectX
+							}
+						}
+					}
+
+					TMap<uint64, int32> RemapEdgeID;
+					{
+						TRACE_CPUPROFILER_EVENT_SCOPE(BuildMeshEdgeVertices);
+						Mesh->BeginGetMeshEdgeVertices();
+
+						//Fill the edge array
+						int32 FbxEdgeCount = Mesh->GetMeshEdgeCount();
+						RemapEdgeID.Reserve(FbxEdgeCount * 2);
+						for (int32 FbxEdgeIndex = 0; FbxEdgeIndex < FbxEdgeCount; ++FbxEdgeIndex)
+						{
+							int32 EdgeStartVertexIndex = -1;
+							int32 EdgeEndVertexIndex = -1;
+							Mesh->GetMeshEdgeVertices(FbxEdgeIndex, EdgeStartVertexIndex, EdgeEndVertexIndex);
+							// Skip invalid edges, i.e. one of the ends is invalid, or degenerated ones
+							if (EdgeStartVertexIndex == -1 || EdgeEndVertexIndex == -1 || EdgeStartVertexIndex == EdgeEndVertexIndex)
+							{
+								//TODO log a user error
+								//UE_LOG(LogFbx, Warning, TEXT("Skipping invalid edge on mesh %s"), *FString(Mesh->GetName()));
+								continue;
+							}
+							FVertexID EdgeVertexStart(EdgeStartVertexIndex + VertexOffset);
+							check(MeshDescription->Vertices().IsValid(EdgeVertexStart));
+							FVertexID EdgeVertexEnd(EdgeEndVertexIndex + VertexOffset);
+							check(MeshDescription->Vertices().IsValid(EdgeVertexEnd));
+							uint64 CompactedKey = (((uint64)EdgeVertexStart.GetValue()) << 32) | ((uint64)EdgeVertexEnd.GetValue());
+							RemapEdgeID.Add(CompactedKey, FbxEdgeIndex);
+							//Add the other edge side
+							CompactedKey = (((uint64)EdgeVertexEnd.GetValue()) << 32) | ((uint64)EdgeVertexStart.GetValue());
+							RemapEdgeID.Add(CompactedKey, FbxEdgeIndex);
+						}
+						//Call this after all GetMeshEdgeIndexForPolygon call this is for optimization purpose.
+						Mesh->EndGetMeshEdgeVertices();
+					}
+
+					{
+						TRACE_CPUPROFILER_EVENT_SCOPE(BuildMeshEdgeIndexForPolygon);
+
+						// Compute and reserve memory to be used for vertex instances
+						{
+							int32 TotalVertexCount = 0;
+							for (int32 PolygonIndex = 0; PolygonIndex < PolygonCount; PolygonIndex++)
+							{
+								TotalVertexCount += Mesh->GetPolygonSize(PolygonIndex);
+							}
+
+							MeshDescription->ReserveNewPolygons(PolygonCount);
+							MeshDescription->ReserveNewVertexInstances(TotalVertexCount);
+							MeshDescription->ReserveNewEdges(TotalVertexCount);
+						}
+
+						bool  bBeginGetMeshEdgeIndexForPolygonCalled = false;
+						bool  bBeginGetMeshEdgeIndexForPolygonRequired = true;
+						int32 CurrentVertexInstanceIndex = 0;
+						int32 SkippedVertexInstance = 0;
+
+						// keep those for all iterations to avoid heap allocations
+						TArray<FVertexInstanceID> CornerInstanceIDs;
+						TArray<FVertexID> CornerVerticesIDs;
+						TArray<FVector, TInlineAllocator<3>> P;
+
+						//Polygons
+						for (int32 PolygonIndex = 0; PolygonIndex < PolygonCount; PolygonIndex++)
+						{
+							int32 PolygonVertexCount = Mesh->GetPolygonSize(PolygonIndex);
+							//Verify if the polygon is degenerate, in this case do not add them
+							{
+								float ComparisonThreshold = SMALL_NUMBER;
+								P.Reset();
+								P.AddUninitialized(PolygonVertexCount);
+								for (int32 CornerIndex = 0; CornerIndex < PolygonVertexCount; CornerIndex++)
+								{
+									const int32 ControlPointIndex = Mesh->GetPolygonVertex(PolygonIndex, CornerIndex);
+									const FVertexID VertexID(VertexOffset + ControlPointIndex);
+									P[CornerIndex] = VertexPositions[VertexID];
+								}
+								check(P.Num() > 2); //triangle is the smallest polygon we can have
+								const FVector Normal = ((P[1] - P[2]) ^ (P[0] - P[2])).GetSafeNormal(ComparisonThreshold);
+								//Check for degenerated polygons, avoid NAN
+								if (Normal.IsNearlyZero(ComparisonThreshold) || Normal.ContainsNaN())
+								{
+									SkippedVertexInstance += PolygonVertexCount;
+									continue;
+								}
+							}
+
+							int32 RealPolygonIndex = PolygonOffset + PolygonIndex;
+							CornerInstanceIDs.Reset();
+							CornerInstanceIDs.AddUninitialized(PolygonVertexCount);
+							CornerVerticesIDs.Reset();
+							CornerVerticesIDs.AddUninitialized(PolygonVertexCount);
+							for (int32 CornerIndex = 0; CornerIndex < PolygonVertexCount; CornerIndex++)
+							{
+								int32 VertexInstanceIndex = VertexInstanceOffset + CurrentVertexInstanceIndex + CornerIndex;
+								int32 RealFbxVertexIndex = SkippedVertexInstance + CurrentVertexInstanceIndex + CornerIndex;
+								const FVertexInstanceID VertexInstanceID(VertexInstanceIndex);
+								CornerInstanceIDs[CornerIndex] = VertexInstanceID;
+								const int32 ControlPointIndex = Mesh->GetPolygonVertex(PolygonIndex, CornerIndex);
+								const FVertexID VertexID(VertexOffset + ControlPointIndex);
+								const FVector VertexPosition = VertexPositions[VertexID];
+								CornerVerticesIDs[CornerIndex] = VertexID;
+
+								FVertexInstanceID AddedVertexInstanceId = MeshDescription->CreateVertexInstance(VertexID);
+
+								//Make sure the Added vertex instance ID is matching the expected vertex instance ID
+								check(AddedVertexInstanceId == VertexInstanceID);
+
+								if (AddedVertexInstanceId.GetValue() != VertexInstanceIndex)
+								{
+									//TODO add an error log
+									//AddTokenizedErrorMessage(FTokenizedMessage::Create(EMessageSeverity::Error, FText::Format(LOCTEXT("Error_CannotCreateVertexInstance", "Cannot create valid vertex instance for mesh '{0}'"), FText::FromString(Mesh->GetName()))), FFbxErrors::StaticMesh_BuildError);
+									return false;
+								}
+
+								//UVs attributes
+								for (int32 UVLayerIndex = 0; UVLayerIndex < FBXUVs.UniqueUVCount; UVLayerIndex++)
+								{
+									FVector2D FinalUVVector(0.0f, 0.0f);
+									if (FBXUVs.LayerElementUV[UVLayerIndex] != NULL)
+									{
+										int UVMapIndex = (FBXUVs.UVMappingMode[UVLayerIndex] == FbxLayerElement::eByControlPoint) ? ControlPointIndex : RealFbxVertexIndex;
+										int32 UVIndex = (FBXUVs.UVReferenceMode[UVLayerIndex] == FbxLayerElement::eDirect) ?
+											UVMapIndex : FBXUVs.LayerElementUV[UVLayerIndex]->GetIndexArray().GetAt(UVMapIndex);
+
+										FbxVector2	UVVector = FBXUVs.LayerElementUV[UVLayerIndex]->GetDirectArray().GetAt(UVIndex);
+										FinalUVVector.X = static_cast<float>(UVVector[0]);
+										FinalUVVector.Y = 1.f - static_cast<float>(UVVector[1]);   //flip the Y of UVs for DirectX
+									}
+									VertexInstanceUVs.Set(AddedVertexInstanceId, UVLayerIndex, FinalUVVector);
+								}
+
+								//Color attribute
+								if (LayerElementVertexColor)
+								{
+									int32 VertexColorMappingIndex = (VertexColorMappingMode == FbxLayerElement::eByControlPoint) ?
+										Mesh->GetPolygonVertex(PolygonIndex, CornerIndex) : (RealFbxVertexIndex);
+
+									int32 VectorColorIndex = (VertexColorReferenceMode == FbxLayerElement::eDirect) ?
+										VertexColorMappingIndex : LayerElementVertexColor->GetIndexArray().GetAt(VertexColorMappingIndex);
+
+									FbxColor VertexColor = LayerElementVertexColor->GetDirectArray().GetAt(VectorColorIndex);
+
+									FColor VertexInstanceColor(
+										uint8(255.f * VertexColor.mRed),
+										uint8(255.f * VertexColor.mGreen),
+										uint8(255.f * VertexColor.mBlue),
+										uint8(255.f * VertexColor.mAlpha)
+									);
+									VertexInstanceColors[AddedVertexInstanceId] = FVector4(FLinearColor(VertexInstanceColor));
+								}
+
+								if (LayerElementNormal)
+								{
+									//normals may have different reference and mapping mode than tangents and binormals
+									int NormalMapIndex = (NormalMappingMode == FbxLayerElement::eByControlPoint) ?
+										ControlPointIndex : RealFbxVertexIndex;
+									int NormalValueIndex = (NormalReferenceMode == FbxLayerElement::eDirect) ?
+										NormalMapIndex : LayerElementNormal->GetIndexArray().GetAt(NormalMapIndex);
+
+									FbxVector4 TempValue = LayerElementNormal->GetDirectArray().GetAt(NormalValueIndex);
+									TempValue = TotalMatrixForNormal.MultT(TempValue);
+									FVector TangentZ = FFbxConvert::ConvertDir(TempValue);
+									VertexInstanceNormals[AddedVertexInstanceId] = TangentZ.GetSafeNormal();
+									//tangents and binormals share the same reference, mapping mode and index array
+									if (bHasNTBInformation)
+									{
+										int TangentMapIndex = (TangentMappingMode == FbxLayerElement::eByControlPoint) ?
+											ControlPointIndex : RealFbxVertexIndex;
+										int TangentValueIndex = (TangentReferenceMode == FbxLayerElement::eDirect) ?
+											TangentMapIndex : LayerElementTangent->GetIndexArray().GetAt(TangentMapIndex);
+
+										TempValue = LayerElementTangent->GetDirectArray().GetAt(TangentValueIndex);
+										TempValue = TotalMatrixForNormal.MultT(TempValue);
+										FVector TangentX = FFbxConvert::ConvertDir(TempValue);
+										VertexInstanceTangents[AddedVertexInstanceId] = TangentX.GetSafeNormal();
+
+										int BinormalMapIndex = (BinormalMappingMode == FbxLayerElement::eByControlPoint) ?
+											ControlPointIndex : RealFbxVertexIndex;
+										int BinormalValueIndex = (BinormalReferenceMode == FbxLayerElement::eDirect) ?
+											BinormalMapIndex : LayerElementBinormal->GetIndexArray().GetAt(BinormalMapIndex);
+
+										TempValue = LayerElementBinormal->GetDirectArray().GetAt(BinormalValueIndex);
+										TempValue = TotalMatrixForNormal.MultT(TempValue);
+										FVector TangentY = -FFbxConvert::ConvertDir(TempValue);
+										VertexInstanceBinormalSigns[AddedVertexInstanceId] = FbxGetBasisDeterminantSign(TangentX.GetSafeNormal(), TangentY.GetSafeNormal(), TangentZ.GetSafeNormal());
+									}
+								}
+							}
+
+							// Check if the polygon just discovered is non-degenerate if we haven't found one yet
+							//TODO check all polygon vertex, not just the first 3 vertex
+							if (!bHasNonDegeneratePolygons)
+							{
+								float TriangleComparisonThreshold = THRESH_POINTS_ARE_SAME;
+								FVector VertexPosition[3];
+								VertexPosition[0] = VertexPositions[CornerVerticesIDs[0]];
+								VertexPosition[1] = VertexPositions[CornerVerticesIDs[1]];
+								VertexPosition[2] = VertexPositions[CornerVerticesIDs[2]];
+								if (!(VertexPosition[0].Equals(VertexPosition[1], TriangleComparisonThreshold)
+									|| VertexPosition[0].Equals(VertexPosition[2], TriangleComparisonThreshold)
+									|| VertexPosition[1].Equals(VertexPosition[2], TriangleComparisonThreshold)))
+								{
+									bHasNonDegeneratePolygons = true;
+								}
+							}
+
+							//
+							// material index
+							//
+							int32 MaterialIndex = 0;
+							if (MaterialCount > 0)
+							{
+								if (LayerElementMaterial)
+								{
+									switch (MaterialMappingMode)
+									{
+										// material index is stored in the IndexArray, not the DirectArray (which is irrelevant with 2009.1)
+										case FbxLayerElement::eAllSame:
+										{
+											MaterialIndex = LayerElementMaterial->GetIndexArray().GetAt(0);
+										}
+										break;
+										case FbxLayerElement::eByPolygon:
+										{
+											MaterialIndex = LayerElementMaterial->GetIndexArray().GetAt(PolygonIndex);
+										}
+										break;
+									}
+								}
+							}
+
+							if (MaterialIndex >= MaterialCount || MaterialIndex < 0)
+							{
+								//TODO log an error message
+								//AddTokenizedErrorMessage(FTokenizedMessage::Create(EMessageSeverity::Warning, LOCTEXT("Error_MaterialIndexInconsistency", "Face material index inconsistency - forcing to 0")), FFbxErrors::Generic_Mesh_MaterialIndexInconsistency);
+								MaterialIndex = 0;
+							}
+							if (!MaterialNames.IsValidIndex(MaterialIndex))
+							{
+								MaterialIndex = 0;
+							}
+							else
+							{
+								//Create a polygon with the 3 vertex instances Add it to the material group
+								FName MaterialName = MaterialNames[MaterialIndex];
+								if (!PolygonGroupMapping.Contains(MaterialIndex))
+								{
+									FPolygonGroupID ExistingPolygonGroup = INDEX_NONE;
+									for (const FPolygonGroupID PolygonGroupID : MeshDescription->PolygonGroups().GetElementIDs())
+									{
+										if (PolygonGroupImportedMaterialSlotNames[PolygonGroupID] == MaterialName)
+										{
+											ExistingPolygonGroup = PolygonGroupID;
+											break;
+										}
+									}
+									if (ExistingPolygonGroup == INDEX_NONE)
+									{
+										ExistingPolygonGroup = MeshDescription->CreatePolygonGroup();
+										PolygonGroupImportedMaterialSlotNames[ExistingPolygonGroup] = MaterialName;
+									}
+									PolygonGroupMapping.Add(MaterialIndex, ExistingPolygonGroup);
+								}
+							}
+
+							// Create polygon edges
+							{
+								// Add the edges of this polygon
+								for (uint32 PolygonEdgeNumber = 0; PolygonEdgeNumber < (uint32)PolygonVertexCount; ++PolygonEdgeNumber)
+								{
+									//Find the matching edge ID
+									uint32 CornerIndices[2];
+									CornerIndices[0] = (PolygonEdgeNumber + 0) % PolygonVertexCount;
+									CornerIndices[1] = (PolygonEdgeNumber + 1) % PolygonVertexCount;
+
+									FVertexID EdgeVertexIDs[2];
+									EdgeVertexIDs[0] = CornerVerticesIDs[CornerIndices[0]];
+									EdgeVertexIDs[1] = CornerVerticesIDs[CornerIndices[1]];
+
+									FEdgeID MatchEdgeId = MeshDescription->GetVertexPairEdge(EdgeVertexIDs[0], EdgeVertexIDs[1]);
+									if (MatchEdgeId == INDEX_NONE)
+									{
+										MatchEdgeId = MeshDescription->CreateEdge(EdgeVertexIDs[0], EdgeVertexIDs[1]);
+									}
+
+									//RawMesh do not have edges, so by ordering the edge with the triangle construction we can ensure back and forth conversion with RawMesh
+									//When raw mesh will be completely remove we can create the edges right after the vertex creation.
+									int32 EdgeIndex = INDEX_NONE;
+									uint64 CompactedKey = (((uint64)EdgeVertexIDs[0].GetValue()) << 32) | ((uint64)EdgeVertexIDs[1].GetValue());
+									if (RemapEdgeID.Contains(CompactedKey))
+									{
+										EdgeIndex = RemapEdgeID[CompactedKey];
+									}
+									else
+									{
+										// Call BeginGetMeshEdgeIndexForPolygon lazily only if we enter the code path calling GetMeshEdgeIndexForPolygon
+										if (bBeginGetMeshEdgeIndexForPolygonRequired)
+										{
+											//Call this before all GetMeshEdgeIndexForPolygon for optimization purpose.
+											//But do not spend time precomputing stuff if the mesh has no edge since 
+											//GetMeshEdgeIndexForPolygon is going to always returns -1 anyway without going into the slow path.
+											if (Mesh->GetMeshEdgeCount() > 0)
+											{
+												Mesh->BeginGetMeshEdgeIndexForPolygon();
+												bBeginGetMeshEdgeIndexForPolygonCalled = true;
+											}
+											bBeginGetMeshEdgeIndexForPolygonRequired = false;
+										}
+
+										EdgeIndex = Mesh->GetMeshEdgeIndexForPolygon(PolygonIndex, PolygonEdgeNumber);
+									}
+
+									if (!EdgeHardnesses[MatchEdgeId])
+									{
+										if (bSmoothingAvailable && SmoothingInfo)
+										{
+											if (SmoothingMappingMode == FbxLayerElement::eByEdge)
+											{
+												int32 lSmoothingIndex = (SmoothingReferenceMode == FbxLayerElement::eDirect) ? EdgeIndex : SmoothingInfo->GetIndexArray().GetAt(EdgeIndex);
+												//Set the hard edges
+												int32 SmoothingFlag = SmoothingInfo->GetDirectArray().GetAt(lSmoothingIndex);
+												EdgeHardnesses[MatchEdgeId] = (SmoothingFlag == 0);
+											}
+											else
+											{
+												EdgeHardnesses[MatchEdgeId] = false;
+												//TODO add an error log
+												//AddTokenizedErrorMessage(FTokenizedMessage::Create(EMessageSeverity::Warning, FText::Format(LOCTEXT("Error_UnsupportedSmoothingGroup", "Unsupported Smoothing group mapping mode on mesh  '{0}'"), FText::FromString(Mesh->GetName()))), FFbxErrors::Generic_Mesh_UnsupportingSmoothingGroup);
+											}
+										}
+										else
+										{
+											//When there is no smoothing group we set all edge to: hard (faceted mesh) for static mesh and smooth for skinned and rigid
+											EdgeHardnesses[MatchEdgeId] = MeshType == EMeshType::Static ? true : false;
+										}
+									}
+								}
+							}
+
+							FPolygonGroupID PolygonGroupID = PolygonGroupMapping[MaterialIndex];
+							// Insert a triangle into the mesh
+							// @note: we only ever import triangulated meshes currently. This could easily change as we have the infrastructure set up for arbitrary ngons.
+							// However, I think either we need to triangulate in the exact same way as Maya/Max if we do that.
+							check(CornerInstanceIDs.Num() == 3);
+							check(PolygonVertexCount == 3);
+							TArray<FEdgeID> NewEdgeIDs;
+							const FTriangleID NewTriangleID = MeshDescription->CreateTriangle(PolygonGroupID, CornerInstanceIDs, &NewEdgeIDs);
+							check(NewEdgeIDs.Num() == 0);
+
+							for (int32 UVLayerIndex = 0; UVLayerIndex < FBXUVs.UniqueUVCount; UVLayerIndex++)
+							{
+								FUVID UVIDs[3] = { INDEX_NONE, INDEX_NONE, INDEX_NONE };
+								if (FBXUVs.LayerElementUV[UVLayerIndex] != NULL)
+								{
+									for (int32 VertexIndex = 0; VertexIndex < 3; ++VertexIndex)
+									{
+										int UVMapIndex = (FBXUVs.UVMappingMode[UVLayerIndex] == FbxLayerElement::eByControlPoint)
+											? Mesh->GetPolygonVertex(PolygonIndex, VertexIndex)
+											: SkippedVertexInstance + CurrentVertexInstanceIndex + VertexIndex;
+										int32 UVIndex = (FBXUVs.UVReferenceMode[UVLayerIndex] == FbxLayerElement::eDirect)
+											? UVMapIndex
+											: FBXUVs.LayerElementUV[UVLayerIndex]->GetIndexArray().GetAt(UVMapIndex);
+
+										if (UVIndex != -1)
+										{
+											UVIDs[VertexIndex] = UVIndex + UVOffsets[UVLayerIndex];
+
+											check(MeshDescription->VertexInstanceAttributes().GetAttribute<FVector2D>(CornerInstanceIDs[VertexIndex], MeshAttribute::VertexInstance::TextureCoordinate, UVLayerIndex) ==
+												  MeshDescription->UVAttributes(UVLayerIndex).GetAttribute<FVector2D>(UVIndex + UVOffsets[UVLayerIndex], MeshAttribute::UV::UVCoordinate));
+										}
+										else
+										{
+											// TODO: what does it mean to have a UV index of -1?
+											// Investigate this case more carefully and handle it properly.
+										}
+									}
+								}
+
+								MeshDescription->SetTriangleUVIndices(NewTriangleID, UVIDs, UVLayerIndex);
+							}
+
+							CurrentVertexInstanceIndex += PolygonVertexCount;
+						}
+
+						//Call this after all GetMeshEdgeIndexForPolygon call this is for optimization purpose.
+						if (bBeginGetMeshEdgeIndexForPolygonCalled)
+						{
+							Mesh->EndGetMeshEdgeIndexForPolygon();
+						}
+
+						MeshDescription->ResumeVertexInstanceIndexing();
+						MeshDescription->ResumeEdgeIndexing();
+						MeshDescription->ResumePolygonIndexing();
+						MeshDescription->ResumePolygonGroupIndexing();
+						MeshDescription->ResumeUVIndexing();
+
+						if (SkippedVertexInstance > 0)
+						{
+							//We must compact the sparse array before reserving new space
+							//When we reserve it will make a hole in the sparse array if the last reserve was not fully use
+							//The importer assume there will be no hole when importing a mesh
+							FElementIDRemappings OutRemappings;
+							MeshDescription->Compact(OutRemappings);
+						}
+					}
+
+					if (MeshType == EMeshType::Skinned || MeshType == EMeshType::Rigid)
+					{
+						if (!ensure(SortedJoints))
+						{
+							//TODO log an error
+							return false;
+						}
+
+						FSkeletalMeshAttributes SkeletalMeshAttributes(*MeshDescription);
+
+						//Add the influence data in the skeletalmesh description
+						TVertexAttributesRef<TArrayAttribute<int32>> VertexInfluenceBones = SkeletalMeshAttributes.GetVertexInfluenceBones();
+						TVertexAttributesRef<TArrayAttribute<float>> VertexInfluenceWeights = SkeletalMeshAttributes.GetVertexInfluenceWeights();
+
+
+						auto AddOneInfluence = [&VertexInfluenceBones, &VertexInfluenceWeights](FVertexID VertexID, int32 BoneIndex, float Weight)
+						{
+							//Add one influence to the vertex if its not already set
+							int32 InfluenceCount = VertexInfluenceBones[VertexID].Num();
+							bool bCreateNewInfluenceEntry = true;
+							for (int32 InfluenceIndex = 0; InfluenceIndex < InfluenceCount; ++InfluenceIndex)
+							{
+								int32 InfluenceBoneIndex = VertexInfluenceBones[VertexID][InfluenceIndex];
+								if (InfluenceBoneIndex == BoneIndex)
+								{
+									//Use the most heavy weight
+									float InfluenceWeight = VertexInfluenceWeights[VertexID][InfluenceIndex];
+									if (Weight > InfluenceWeight)
+									{
+										VertexInfluenceWeights[VertexID][InfluenceIndex] = Weight;
+									}
+									bCreateNewInfluenceEntry = false;
+									break;
+								}
+							}
+							if (bCreateNewInfluenceEntry)
+							{
+								ensure(InfluenceCount == VertexInfluenceBones[VertexID].Add(BoneIndex));
+								ensure(InfluenceCount == VertexInfluenceWeights[VertexID].Add(Weight));
+							}
+						};
+
+						if (MeshType == EMeshType::Skinned)
+						{
+							FbxSkin* Skin = (FbxSkin*)Mesh->GetDeformer(0, FbxDeformer::eSkin);
+							if (!ensure(Skin))
+							{
+								//TODO log an error
+								return false;
+							}
+							// create influences for each cluster
+							for (int32 ClusterIndex = 0; ClusterIndex < Skin->GetClusterCount(); ClusterIndex++)
+							{
+								FbxCluster* Cluster = Skin->GetCluster(ClusterIndex);
+								// When Maya plug-in exports rigid binding, it will generate "CompensationCluster" for each ancestor links.
+								// FBX writes these "CompensationCluster" out. The CompensationCluster also has weight 1 for vertices.
+								// Unreal importer should skip these clusters.
+								if (!Cluster || (FCStringAnsi::Strcmp(Cluster->GetUserDataID(), "Maya_ClusterHint") == 0 && FCStringAnsi::Strcmp(Cluster->GetUserData(), "CompensationCluster") == 0))
+								{
+									continue;
+								}
+
+								FbxNode* Link = Cluster->GetLink();
+								// find the bone index
+								int32 BoneIndex = -1;
+								SortedJoints->Find(Link, BoneIndex);
+
+								//	get the vertex indices
+								int32 ControlPointIndicesCount = Cluster->GetControlPointIndicesCount();
+								int32* ControlPointIndices = Cluster->GetControlPointIndices();
+								double* Weights = Cluster->GetControlPointWeights();
+
+								//	for each vertex index in the cluster
+								for (int32 ControlPointIndex = 0; ControlPointIndex < ControlPointIndicesCount; ++ControlPointIndex)
+								{
+									FVertexID VertexID(ControlPointIndices[ControlPointIndex] + VertexOffset);
+									if (!ensure(MeshDescription->IsVertexValid(VertexID)))
+									{
+										//Invalid Influence
+										continue;
+									}
+									float Weight = static_cast<float>(Weights[ControlPointIndex]);
+
+									AddOneInfluence(VertexID, BoneIndex, Weight);
+								}
+							}
+						}
+						else // for rigid mesh
+						{
+							// find the bone index, the bone is the node itself
+							int32 BoneIndex = -1;
+							SortedJoints->Find(MeshNode, BoneIndex);
+
+							if (BoneIndex == -1 && SortedJoints->Num() > 0)
+							{
+								//When we import geometry only, we want to hook all the influence to a particular bone
+								BoneIndex = 0;
+							}
+							//	for each vertex in the mesh
+							for (int32 ControlPointIndex = 0; ControlPointIndex < VertexCount; ++ControlPointIndex)
+							{
+								FVertexID VertexID(ControlPointIndex + VertexOffset);
+								if (!ensure(MeshDescription->IsVertexValid(VertexID)))
+								{
+									//Invalid Influence
+									continue;
+								}
+								const float Weight = 1.0f;
+								AddOneInfluence(VertexID, BoneIndex, Weight);
+							}
+						}
+
+						//Reorder by weight
+						{
+							TArray<int32> BoneIndexes;
+							TArray<float> BoneWeights;
+							TArray<int32> RemapIndexes;
+							TArray<float> RemapWeights;
+							for (const FVertexID VertexID : MeshDescription->Vertices().GetElementIDs())
+							{
+								//Add one influence to the vertex if its not already set
+								int32 InfluenceCount = VertexInfluenceBones[VertexID].Num();
+								BoneIndexes.Reset(InfluenceCount);
+								BoneIndexes.AddZeroed(InfluenceCount);
+								BoneWeights.Reset(InfluenceCount);
+								BoneWeights.AddZeroed(InfluenceCount);
+								RemapIndexes.Reset(InfluenceCount);
+								RemapIndexes.AddZeroed(InfluenceCount);
+								RemapWeights.Reset(InfluenceCount);
+								RemapWeights.AddZeroed(InfluenceCount);
+								for (int32 InfluenceIndex = 0; InfluenceIndex < InfluenceCount; ++InfluenceIndex)
+								{
+									int32 InfluenceBoneIndex = VertexInfluenceBones[VertexID][InfluenceIndex];
+									BoneIndexes[InfluenceIndex] = InfluenceBoneIndex;
+									//Use the most heavy weight
+									float InfluenceWeight = VertexInfluenceWeights[VertexID][InfluenceIndex];
+									BoneWeights[InfluenceIndex] = InfluenceWeight;
+									RemapIndexes[InfluenceIndex] = InfluenceIndex;
+									RemapWeights[InfluenceIndex] = InfluenceWeight;
+								}
+
+								for (int32 InfluenceIndex = 0; InfluenceIndex < InfluenceCount; ++InfluenceIndex)
+								{
+									if (InfluenceIndex == 0)
+									{
+										continue;
+									}
+									float CurrentWeight = BoneWeights[InfluenceIndex];
+									for (int32 RemapIndex = InfluenceIndex - 1; RemapIndex >= 0; RemapIndex--)
+									{
+										if (CurrentWeight > BoneWeights[RemapIndexes[RemapIndex]])
+										{
+											//Swap
+											int32 SwapIndex = RemapIndexes[RemapIndex + 1];
+											float SwapValue = RemapWeights[RemapIndex + 1];
+											RemapIndexes[RemapIndex + 1] = RemapIndexes[RemapIndex];
+											RemapWeights[RemapIndex + 1] = RemapWeights[RemapIndex];
+											RemapIndexes[RemapIndex] = SwapIndex;
+											RemapWeights[RemapIndex] = SwapValue;
+										}
+										else
+										{
+											break;
+										}
+									}
+								}
+								//Push the value into the mesh description
+								for (int32 InfluenceIndex = 0; InfluenceIndex < InfluenceCount; ++InfluenceIndex)
+								{
+									BoneIndexes[InfluenceIndex] = RemapIndexes[InfluenceIndex];
+									BoneWeights[InfluenceIndex] = RemapWeights[InfluenceIndex];
+								}
+							}
+						}
+					}
+				}
+				// needed?
+				FBXUVs.Cleanup();
+
+				if (!bHasNonDegeneratePolygons)
+				{
+					//TODO log an error we have imported no valid triangle in this fbx mesh node
+					//CreateTokenizedErrorForDegeneratedPart(this, StaticMesh->GetName(), FbxNodeName);
+				}
+
+				return bHasNonDegeneratePolygons;
+			}
+
+			FbxAMatrix FMeshDescriptionImporter::ComputeNodeMatrix(FbxNode* Node)
+			{
+				FbxAMatrix Geometry;
+				FbxVector4 Translation, Rotation, Scaling;
+				Translation = Node->GetGeometricTranslation(FbxNode::eSourcePivot);
+				Rotation = Node->GetGeometricRotation(FbxNode::eSourcePivot);
+				Scaling = Node->GetGeometricScaling(FbxNode::eSourcePivot);
+				Geometry.SetT(Translation);
+				Geometry.SetR(Rotation);
+				Geometry.SetS(Scaling);
+				//For Single Matrix situation, obtain transfrom matrix from eDESTINATION_SET, which include pivot offsets and pre/post rotations.
+				FbxAMatrix& GlobalTransform = SDKScene->GetAnimationEvaluator()->GetNodeGlobalTransform(Node);
+				//We must always add the geometric transform. Only Max use the geometric transform which is an offset to the local transform of the node
+				FbxAMatrix NodeMatrix = GlobalTransform * Geometry;
+				return NodeMatrix;
+			}
+
+			bool FMeshDescriptionImporter::IsOddNegativeScale(FbxAMatrix& TotalMatrix)
+			{
+				FbxVector4 Scale = TotalMatrix.GetS();
+				int32 NegativeNum = 0;
+
+				if (Scale[0] < 0) NegativeNum++;
+				if (Scale[1] < 0) NegativeNum++;
+				if (Scale[2] < 0) NegativeNum++;
+
+				return NegativeNum == 1 || NegativeNum == 3;
+			}
+		} //ns Private
+	} //ns Interchange
+}//ns UE
diff --git a/Engine/Plugins/Runtime/Interchange/Source/Parsers/Fbx/Private/FbxMesh.h b/Engine/Plugins/Runtime/Interchange/Source/Parsers/Fbx/Private/FbxMesh.h
new file mode 100644
--- /dev/null
+++ b/Engine/Plugins/Runtime/Interchange/Source/Parsers/Fbx/Private/FbxMesh.h
@@ -0,0 +1,61 @@
+// Copyright Epic Games, Inc. All Rights Reserved.
+
+#pragma once
+
+#include "CoreMinimal.h"
+#include "FbxHelper.h"
+#include "FbxInclude.h"
+
+/** Forward declarations */
+struct FMeshDescription;
+
+
+namespace UE
+{
+	namespace Interchange
+	{
+		namespace Private
+		{
+			class FMeshDescriptionImporter
+			{
+			public:
+				FMeshDescriptionImporter(FMeshDescription* InMeshDescription, FbxNode* InMeshNode, FbxScene* InSDKScene, FbxGeometryConverter* InSDKGeometryConverter);
+				bool FillStaticMeshDescriptionFromFbxMesh();
+				bool FillSkinnedMeshDescriptionFromFbxMesh(TArray<FbxNode*>* SortedJoints);
+			private:
+				
+				enum class EMeshType : uint8
+				{
+					None = 0, //No mesh type to import
+					Static = 1, //static mesh
+					Skinned = 2, //skinned mesh with joints
+					Rigid = 3  //rigid mesh (joints will be created from the geometry transform)
+				};
+
+				bool FillMeshDescriptionFromFbxMesh(EMeshType MeshType, TArray<FbxNode*>* SortedJoints);
+				FbxAMatrix ComputeNodeMatrix(FbxNode* Node);
+				bool IsOddNegativeScale(FbxAMatrix& TotalMatrix);
+				
+				//TODO move the real function from RenderCore to FVector, so we do not have to add render core to compute such a simple thing
+				float FbxGetBasisDeterminantSign(const FVector& XAxis, const FVector& YAxis, const FVector& ZAxis)
+				{
+					FMatrix Basis(
+						FPlane(XAxis, 0),
+						FPlane(YAxis, 0),
+						FPlane(ZAxis, 0),
+						FPlane(0, 0, 0, 1)
+					);
+					return (Basis.Determinant() < 0) ? -1.0f : +1.0f;
+				}
+
+
+				FMeshDescription* MeshDescription;
+				FbxNode* MeshNode;
+				FbxScene* SDKScene;
+				FbxGeometryConverter* SDKGeometryConverter;
+				bool bIsStaticMesh;
+				bool bInitialized = false;
+			};
+		}//ns Private
+	}//ns Interchange
+}//ns UE
diff --git a/Engine/Plugins/Runtime/Interchange/Source/Parsers/Fbx/Private/FbxScene.cpp b/Engine/Plugins/Runtime/Interchange/Source/Parsers/Fbx/Private/FbxScene.cpp
--- a/Engine/Plugins/Runtime/Interchange/Source/Parsers/Fbx/Private/FbxScene.cpp
+++ b/Engine/Plugins/Runtime/Interchange/Source/Parsers/Fbx/Private/FbxScene.cpp
@@ -4,6 +4,7 @@
 
 #include "CoreMinimal.h"
 #include "FbxConvert.h"
+#include "FbxHelper.h"
 #include "FbxInclude.h"
 #include "FbxMaterial.h"
 #include "InterchangeSceneNode.h"
@@ -19,61 +20,36 @@ namespace UE
 			{
 				UInterchangeSceneNode* AddHierarchyRecursively(FbxNode* Node, FbxScene* SDKScene, UInterchangeBaseNodeContainer& NodeContainer, TArray<FString>& JSonErrorMessages)
 				{
-					FString NodeName = FFbxConvert::MakeString(FFbxConvert::MakeName(Node->GetName()));
-					UInterchangeSceneNode* UnrealNode = FFbxScene::CreateTransformNode(NodeContainer, NodeName, JSonErrorMessages);
+					FString NodeName = FFbxHelper::GetFbxObjectName(Node);
+					FString NodeUniqueID = FFbxHelper::GetFbxNodeHierarchyName(Node);
+					
+					UInterchangeSceneNode* UnrealNode = nullptr;
+
+					int32 AttributeCount = Node->GetNodeAttributeCount();
+					for (int32 AttributeIndex = 0; AttributeIndex < AttributeCount; ++AttributeIndex)
+					{
+						const FbxNodeAttribute* NodeAttribute = Node->GetNodeAttributeByIndex(AttributeIndex);
+						switch (NodeAttribute->GetAttributeType())
+						{
+							case FbxNodeAttribute::eLODGroup:
+							case FbxNodeAttribute::eMesh:
+							case FbxNodeAttribute::eSkeleton:
+							case FbxNodeAttribute::eCamera:
+							case FbxNodeAttribute::eCameraSwitcher:
+							case FbxNodeAttribute::eLight:
+							default:
+								//Add a transform node
+								UnrealNode = FFbxScene::CreateTransformNode(NodeContainer, NodeName, NodeUniqueID, JSonErrorMessages);
+								break;
+						}
+					}
 					if (UnrealNode)
 					{
 						NodeContainer.AddNode(UnrealNode);
-						FTransform LocalTransform = FFbxConvert::GetTransform(Node->EvaluateLocalTransform());
+						FTransform LocalTransform = FFbxConvert::ConvertTransform(Node->EvaluateLocalTransform());
 						UnrealNode->SetCustomLocalTransform(LocalTransform);
-						FTransform GlobalTransform = FFbxConvert::GetTransform(Node->EvaluateGlobalTransform());
+						FTransform GlobalTransform = FFbxConvert::ConvertTransform(Node->EvaluateGlobalTransform());
 						UnrealNode->SetCustomGlobalTransform(GlobalTransform);
-
-						int32 AttributeCount = Node->GetNodeAttributeCount();
-						for (int32 AttributeIndex = 0; AttributeIndex < AttributeCount; ++AttributeIndex)
-						{
-							const FbxNodeAttribute* NodeAttribute = Node->GetNodeAttributeByIndex(AttributeIndex);
-							switch (NodeAttribute->GetAttributeType())
-							{
-								case FbxNodeAttribute::eLODGroup:
-									//Add LOD group node
-									//UnrealNode = FFbxScene::CreateLodGroupNode(NodeContainer, NodeName, JSonErrorMessages);
-									break;
-								case FbxNodeAttribute::eMesh:
-									//Add a mesh node
-									if (Node->GetMesh()->GetDeformerCount(FbxDeformer::eSkin) > 0)
-									{
-										//Skeletal mesh node
-										//UnrealNode = FFbxSkeletalMesh::CreateSkeletalMeshNode(NodeContainer, NodeName, JSonErrorMessages);
-									}
-									else
-									{
-										//Static mesh node
-										//UnrealNode = FFbxStaticMesh::CreateStaticMeshNode(NodeContainer, NodeName, JSonErrorMessages);
-									}
-									break;
-								case FbxNodeAttribute::eSkeleton:
-									//Add a joint node
-									//UnrealNode = FFbxSkeletalMesh::CreateJointNode(NodeContainer, NodeName, JSonErrorMessages);
-									break;
-								case FbxNodeAttribute::eCamera:
-									//Add a camera node
-									//UnrealNode = FFbxCamera::CreateCameraNode(NodeContainer, NodeName, JSonErrorMessages);
-									break;
-								case FbxNodeAttribute::eCameraSwitcher:
-									//Add a camera switcher node
-									//UnrealNode = FFbxCamera::CreateCameraSwitcherNode(NodeContainer, NodeName, JSonErrorMessages);
-									break;
-								case FbxNodeAttribute::eLight:
-									//Add a light node
-									//UnrealNode = FFbxLight::CreateLightNode(NodeContainer, NodeName, JSonErrorMessages);
-									break;
-								default:
-									//Add a transform node
-									UnrealNode = FFbxScene::CreateTransformNode(NodeContainer, NodeName, JSonErrorMessages);
-									break;
-							}
-						}
 						NodeContainer.AddNode(UnrealNode);
 						//Add the dependencies of the material on the correct order
 						FFbxMaterial::AddAllNodeMaterials(UnrealNode, Node, NodeContainer, JSonErrorMessages);
@@ -93,10 +69,10 @@ namespace UE
 				}
 			} //ns Scene
 
-			UInterchangeSceneNode* FFbxScene::CreateTransformNode(UInterchangeBaseNodeContainer& NodeContainer, const FString& NodeName, TArray<FString>& JSonErrorMessages)
+			UInterchangeSceneNode* FFbxScene::CreateTransformNode(UInterchangeBaseNodeContainer& NodeContainer, const FString& NodeName, const FString& NodeUniqueID, TArray<FString>& JSonErrorMessages)
 			{
-				FName DisplayLabel = *NodeName;
-				FName NodeUID(*NodeName);
+				FName DisplayLabel(*NodeName);
+				FName NodeUID(*NodeUniqueID);
 				UInterchangeSceneNode* TransformNode = NewObject<UInterchangeSceneNode>(&NodeContainer, NAME_None);
 				if (!ensure(TransformNode))
 				{
diff --git a/Engine/Plugins/Runtime/Interchange/Source/Parsers/Fbx/Private/FbxScene.h b/Engine/Plugins/Runtime/Interchange/Source/Parsers/Fbx/Private/FbxScene.h
--- a/Engine/Plugins/Runtime/Interchange/Source/Parsers/Fbx/Private/FbxScene.h
+++ b/Engine/Plugins/Runtime/Interchange/Source/Parsers/Fbx/Private/FbxScene.h
@@ -19,7 +19,7 @@ namespace UE
 			{
 			public:
 				static void AddHierarchy(FbxScene* SDKScene, UInterchangeBaseNodeContainer& NodeContainer, TArray<FString>& JSonErrorMessages);
-				static UInterchangeSceneNode* CreateTransformNode(UInterchangeBaseNodeContainer& NodeContainer, const FString& NodeName, TArray<FString>& JSonErrorMessages);
+				static UInterchangeSceneNode* CreateTransformNode(UInterchangeBaseNodeContainer& NodeContainer, const FString& NodeName, const FString& NodeUniqueID, TArray<FString>& JSonErrorMessages);
 			};
 		}//ns Private
 	}//ns Interchange
diff --git a/Engine/Plugins/Runtime/Interchange/Source/Parsers/Fbx/Private/FbxSkeletalMesh.cpp b/Engine/Plugins/Runtime/Interchange/Source/Parsers/Fbx/Private/FbxSkeletalMesh.cpp
new file mode 100644
--- /dev/null
+++ b/Engine/Plugins/Runtime/Interchange/Source/Parsers/Fbx/Private/FbxSkeletalMesh.cpp
@@ -0,0 +1,316 @@
+// Copyright Epic Games, Inc. All Rights Reserved.
+
+#include "FbxSkeletalMesh.h"
+
+#include "CoreMinimal.h"
+#include "FbxConvert.h"
+#include "FbxHelper.h"
+#include "FbxInclude.h"
+#include "FbxMesh.h"
+#include "InterchangeMaterialNode.h"
+#include "InterchangeSkeletonNode.h"
+#include "InterchangeSkeletalMeshNode.h"
+#include "InterchangeSkeletalMeshLodDataNode.h"
+#include "InterchangeTextureNode.h"
+#include "MeshDescription.h"
+#include "Misc/FileHelper.h"
+#include "Nodes/InterchangeBaseNodeContainer.h"
+#include "Serialization/LargeMemoryWriter.h"
+#include "SkeletalMeshAttributes.h"
+
+#define GeneratedLODNameSuffix "_GeneratedLOD_"
+
+#define INTERCHANGE_MAX_TEXCOORDS 4
+
+namespace UE
+{
+	namespace Interchange
+	{
+		namespace Private
+		{
+			bool FSkeletalMeshGeometryPayload::AddFbxMeshToMeshDescription(int32 MeshIndex, FMeshDescription& SkeletalMeshDescription, FSkeletalMeshAttributes& SkeletalMeshAttribute, TArray<FString>& JSonErrorMessages)
+			{
+				if (!ensure(SDKScene) || !ensure(SDKGeometryConverter))
+				{
+					return false;
+				}
+
+				if (!ensure(SkelMeshNodeArray.IsValidIndex(MeshIndex)))
+				{
+					return false;
+				}
+
+				FbxNode* Node = SkelMeshNodeArray[MeshIndex];
+				if (!ensure(Node))
+				{
+					return false;
+				}
+				FbxMesh* Mesh = Node->GetMesh();
+				if (!ensure(Mesh))
+				{
+					return false;
+				}
+				FbxSkin* Skin = (FbxSkin*)Mesh->GetDeformer(0, FbxDeformer::eSkin);
+				if (!ensure(Skin))
+				{
+					return false;
+				}
+
+				FMeshDescriptionImporter MeshDescriptionImporter(&SkeletalMeshDescription, Node, SDKScene, SDKGeometryConverter);
+				if (!MeshDescriptionImporter.FillSkinnedMeshDescriptionFromFbxMesh(&SortedJoints))
+				{
+					return false;
+				}
+
+
+				//////////////////////////////////////////////////////////////////////////
+				//////////////////////////////////////////////////////////////////////////
+				//////////////////////////////////////////////////////////////////////////
+
+				return true;
+			}
+
+			bool FSkeletalMeshGeometryPayload::FetchPayloadToFile(const FString& PayloadFilepath, TArray<FString>& JSonErrorMessages)
+			{
+				if (!ensure(SDKScene != nullptr))
+				{
+					JSonErrorMessages.Add(TEXT("{\"Msg\" : {\"Type\" : \"Error\",\n\"Msg\" : \"Cannot fetch fbx skeletalmesh geometry payload because the fbx scene is null\"}}"));
+					return false;
+				}
+
+				//TODO move MAX_BONES define outside of engine so we can use it here
+				if (SortedJoints.Num() > 65536)
+				{
+					JSonErrorMessages.Add(TEXT("{\"Msg\" : {\"Type\" : \"Error\",\n\"Msg\" : \"Cannot fetch skeletal payload because it use more then 65536 bones\"}}"));
+					return false;
+				}
+
+				//Retreive all meshes point by SkelMeshNodeArray into a FMeshDescription using the skeletalmesh attributes
+				FMeshDescription SkeletalMeshDescription;
+				FSkeletalMeshAttributes SkeletalMeshAttribute(SkeletalMeshDescription);
+				SkeletalMeshAttribute.Register();
+
+				for (int32 MeshIndex = 0; MeshIndex < SkelMeshNodeArray.Num(); ++MeshIndex)
+				{
+					if (!AddFbxMeshToMeshDescription(MeshIndex, SkeletalMeshDescription, SkeletalMeshAttribute, JSonErrorMessages))
+					{
+						continue;
+					}
+				}
+
+				//Dump the MeshDescription to a file
+				{
+					FLargeMemoryWriter Ar;
+					SkeletalMeshDescription.Serialize(Ar);
+					uint8* ArchiveData = Ar.GetData();
+					int64 ArchiveSize = Ar.TotalSize();
+					TArray64<uint8> Buffer(ArchiveData, ArchiveSize);
+					FFileHelper::SaveArrayToFile(Buffer, *PayloadFilepath);
+				}
+
+				return true;
+			}
+
+			void FFbxSkeletalMesh::AddAllSceneSkeletalMeshes(FbxScene* SDKScene, FbxGeometryConverter* SDKGeometryConverter, UInterchangeBaseNodeContainer& NodeContainer, TArray<FString>& JSonErrorMessages, TMap<FString, TSharedPtr<FPayloadContextBase>>& PayloadContexts, const FString& SourceFilename)
+			{
+				TArray<TArray<FbxNode*>> SkeletalMeshes;
+				const bool bCombineSkeletalMesh = false;
+				const bool bForceFindRigid = false;
+				FFbxHelper::FindSkeletalMeshes(SDKScene, SkeletalMeshes, bCombineSkeletalMesh, bForceFindRigid);
+				int32 SkeletalMeshCount = SkeletalMeshes.Num();
+				for (int32 SkeletalMeshIndex = 0; SkeletalMeshIndex < SkeletalMeshCount; ++SkeletalMeshIndex)
+				{
+					TArray<FbxNode*>& NodeArray = SkeletalMeshes[SkeletalMeshIndex];
+					int32 MaxNumberOfLOD = 1;
+					for (int32 NodeIndex = 0; NodeIndex < NodeArray.Num(); NodeIndex++)
+					{
+						FbxNode* Node = NodeArray[NodeIndex];
+						if (Node->GetNodeAttribute() && Node->GetNodeAttribute()->GetAttributeType() == FbxNodeAttribute::eLODGroup)
+						{
+							// get max LODgroup level
+							if (MaxNumberOfLOD < Node->GetChildCount())
+							{
+								MaxNumberOfLOD = Node->GetChildCount();
+							}
+						}
+					}
+					
+					UInterchangeSkeletalMeshNode* SkeletalMeshNode = nullptr;
+
+					bool bOverrideSkeletalMeshName = false;
+					FString OverrideSkeletalMeshName;
+					if (SkeletalMeshCount == 1)
+					{
+						bOverrideSkeletalMeshName = true;
+						OverrideSkeletalMeshName = FPaths::GetBaseFilename(SourceFilename);
+					}
+
+					for (int32 LODIndex = 0; LODIndex < MaxNumberOfLOD; LODIndex++)
+					{
+						TArray<FbxNode*> SkelMeshNodeArray;
+						for (int32 j = 0; j < NodeArray.Num(); j++)
+						{
+							FbxNode* Node = NodeArray[j];
+							if (Node->GetNodeAttribute() && Node->GetNodeAttribute()->GetAttributeType() == FbxNodeAttribute::eLODGroup)
+							{
+								TArray<FbxNode*> NodeInLod;
+								if (Node->GetChildCount() > LODIndex)
+								{
+									FFbxHelper::FindAllLODGroupNode(NodeInLod, Node, LODIndex);
+								}
+								else // in less some LODGroups have less level, use the last level
+								{
+									FFbxHelper::FindAllLODGroupNode(NodeInLod, Node, Node->GetChildCount() - 1);
+								}
+
+								for (FbxNode* MeshNode : NodeInLod)
+								{
+									SkelMeshNodeArray.Add(MeshNode);
+								}
+							}
+							else
+							{
+								SkelMeshNodeArray.Add(Node);
+							}
+						}
+						if (SkelMeshNodeArray.Num() <= 0)
+						{
+							//Nothing to import for this LOD
+							continue;
+						}
+						//Find the skeleton for this skeletal mesh LOD
+						TArray<FbxNode*> SortedJoints;
+						FbxArray<FbxAMatrix> LocalsPerLink;
+						if (!FFbxHelper::FindSkeletonJoints(SDKScene, SkelMeshNodeArray, SortedJoints, LocalsPerLink))
+						{
+							continue;
+						}
+						//A skeleton must have at least one joint
+						if (SortedJoints.Num() <= 0)
+						{
+							continue;
+						}
+						FString SkeletonName = (bOverrideSkeletalMeshName ? OverrideSkeletalMeshName : FFbxHelper::GetFbxObjectName(SkelMeshNodeArray[0])) + TEXT("_Skeleton");
+						if (LODIndex > 0)
+						{
+							SkeletonName += TEXT("_") + FString::FromInt(LODIndex);
+						}
+
+						FString SkeletonUniqueID = TEXT("\\Skeleton\\") + FFbxHelper::GetFbxNodeHierarchyName(SortedJoints[0]);
+
+						UInterchangeSkeletonNode* SkeletonNode = Cast<UInterchangeSkeletonNode>(NodeContainer.GetNode(*SkeletonUniqueID));
+						if (!SkeletonNode)
+						{
+							JSonErrorMessages.Add(TEXT("{\"Msg\" : {\"Type\" : \"Error\",\n\"Msg\" : \"Missing fbx skeleton dependency to import a skeletal mesh\"}}"));
+							//Skip this LOD since the skeleton is already created
+							continue;
+						}
+						//The base LOD will create the SkeletalMeshNode
+						FbxNode* RootNode = SkelMeshNodeArray[0];
+						if (LODIndex == 0)
+						{
+							FString SkeletalMeshName = (bOverrideSkeletalMeshName ? OverrideSkeletalMeshName : FFbxHelper::GetFbxObjectName(RootNode));
+							FString SkeletalMeshUniqueID = TEXT("\\SkeletalMesh\\") + FFbxHelper::GetFbxNodeHierarchyName(RootNode);
+							SkeletalMeshNode = CreateSkeletalMeshNode(NodeContainer, SkeletalMeshName, SkeletalMeshUniqueID, JSonErrorMessages);
+						}
+						if (!ensure(SkeletalMeshNode))
+						{
+							JSonErrorMessages.Add(TEXT("{\"Msg\" : {\"Type\" : \"Error\",\n\"Msg\" : \"Missing fbx skeletal mesh base LOD\"}}"));
+							break;
+						}
+						FString SkeletalMeshLodDataName = TEXT("LodData") + FString::FromInt(LODIndex);
+						FString LODDataPrefix = TEXT("\\LodData") + (LODIndex > 0 ? FString::FromInt(LODIndex) : TEXT("")) + TEXT("\\");
+						FString SkeletalMeshLodDataUniqueID = LODDataPrefix + FFbxHelper::GetFbxNodeHierarchyName(RootNode);
+						//Add the data for the LOD (skeleton Unique ID and all the mesh node fbx path, so we can find them when we will create the payload data)
+						UInterchangeSkeletalMeshLodDataNode* LodDataNode = CreateSkeletalMeshLodDataNode(NodeContainer, SkeletalMeshLodDataName, SkeletalMeshLodDataUniqueID, JSonErrorMessages);
+						LodDataNode->SetParentUID(SkeletalMeshNode->GetUniqueID());
+						LodDataNode->SetCustomSkeletonID(*SkeletonUniqueID);
+						//Add The node unique ids (the node full path) in the LodDataNode
+						//We will use this information to retrieve the payload
+						FSHA1 Sha;
+						
+
+						for (FbxNode* MeshNode : SkelMeshNodeArray)
+						{
+							FString MeshNodePath = FFbxHelper::GetFbxNodeHierarchyName(MeshNode);
+						
+							int32 NodeMaterialCount = MeshNode->GetMaterialCount();
+							for (int32 MaterialIndex = 0; MaterialIndex < NodeMaterialCount; ++MaterialIndex)
+							{
+								FbxSurfaceMaterial* SurfaceMaterial = MeshNode->GetMaterial(MaterialIndex);
+								if (!SurfaceMaterial)
+								{
+									continue;
+								}
+								//Find an existing material node
+								FString MaterialName = FFbxHelper::GetFbxObjectName(SurfaceMaterial);
+								FName NodeUID(*MaterialName);
+								UInterchangeBaseNode* MaterialNode = NodeContainer.GetNode(NodeUID);
+								if (MaterialNode != nullptr && Cast<UInterchangeMaterialNode>(MaterialNode))
+								{
+									//Set a dependency on this material
+									SkeletalMeshNode->SetDependencyUID(NodeUID);
+								}
+							}
+
+							TArray<TCHAR> IDArray = MeshNodePath.GetCharArray();
+							Sha.Update((uint8*)IDArray.GetData(), IDArray.Num()* IDArray.GetTypeSize());
+						}
+						Sha.Final();
+						// Retrieve the hash and use it to construct a pseudo-GUID. 
+						uint32 Hash[5];
+						Sha.GetHash((uint8*)Hash);
+						FString PayloadKey = FGuid(Hash[0] ^ Hash[4], Hash[1], Hash[2], Hash[3]).ToString(EGuidFormats::Base36Encoded);
+						LodDataNode->AddTranslatorMeshKey(*PayloadKey);
+						if (ensure(!PayloadContexts.Contains(PayloadKey)))
+						{
+							TSharedPtr<FSkeletalMeshGeometryPayload> GeoPayload = MakeShared<FSkeletalMeshGeometryPayload>();
+							GeoPayload->SkelMeshNodeArray = SkelMeshNodeArray;
+							GeoPayload->SortedJoints = SortedJoints;
+							GeoPayload->SDKScene = SDKScene;
+							GeoPayload->SDKGeometryConverter = SDKGeometryConverter;
+							PayloadContexts.Add(PayloadKey, GeoPayload);
+						}
+
+						SkeletalMeshNode->AddLodDataUniqueId(*SkeletalMeshLodDataUniqueID);
+						//Add every LOD skeleton has a dependency for the skeletal mesh node
+						//The Skeleton factory will find or create only one skeleton for all LODs
+						SkeletalMeshNode->SetDependencyUID(*SkeletonUniqueID);
+					}
+				}
+			}
+
+			UInterchangeSkeletalMeshNode* FFbxSkeletalMesh::CreateSkeletalMeshNode(UInterchangeBaseNodeContainer& NodeContainer, const FString& NodeName, const FString& NodeUniqueID, TArray<FString>& JSonErrorMessages)
+			{
+				FName DisplayLabel(*NodeName);
+				FName NodeUID(*NodeUniqueID);
+				UInterchangeSkeletalMeshNode* SkeletalMeshNode = NewObject<UInterchangeSkeletalMeshNode>(&NodeContainer, NAME_None);
+				if (!ensure(SkeletalMeshNode))
+				{
+					JSonErrorMessages.Add(TEXT("{\"Msg\" : {\"Type\" : \"Error\",\n\"Msg\" : \"Cannot allocate a node when importing fbx\"}}"));
+					return nullptr;
+				}
+				// Creating a UMaterialInterface
+				SkeletalMeshNode->InitializeSkeletalMeshNode(NodeUID, DisplayLabel, TEXT("SkeletalMesh"));
+				NodeContainer.AddNode(SkeletalMeshNode);
+				return SkeletalMeshNode;
+			}
+
+			UInterchangeSkeletalMeshLodDataNode* FFbxSkeletalMesh::CreateSkeletalMeshLodDataNode(UInterchangeBaseNodeContainer& NodeContainer, const FString& NodeName, const FString& NodeUniqueID, TArray<FString>& JSonErrorMessages)
+			{
+				FName DisplayLabel(*NodeName);
+				FName NodeUID(*NodeUniqueID);
+				UInterchangeSkeletalMeshLodDataNode* SkeletalMeshLodDataNode = NewObject<UInterchangeSkeletalMeshLodDataNode>(&NodeContainer, NAME_None);
+				if (!ensure(SkeletalMeshLodDataNode))
+				{
+					JSonErrorMessages.Add(TEXT("{\"Msg\" : {\"Type\" : \"Error\",\n\"Msg\" : \"Cannot allocate a skeletal mesh lod data node when importing fbx\"}}"));
+					return nullptr;
+				}
+				// Creating a UMaterialInterface
+				SkeletalMeshLodDataNode->InitializeNode(NodeUID, DisplayLabel);
+				NodeContainer.AddNode(SkeletalMeshLodDataNode);
+				return SkeletalMeshLodDataNode;
+			}
+		} //ns Private
+	} //ns Interchange
+}//ns UE
diff --git a/Engine/Plugins/Runtime/Interchange/Source/Parsers/Fbx/Private/FbxSkeletalMesh.h b/Engine/Plugins/Runtime/Interchange/Source/Parsers/Fbx/Private/FbxSkeletalMesh.h
new file mode 100644
--- /dev/null
+++ b/Engine/Plugins/Runtime/Interchange/Source/Parsers/Fbx/Private/FbxSkeletalMesh.h
@@ -0,0 +1,52 @@
+// Copyright Epic Games, Inc. All Rights Reserved.
+
+#pragma once
+
+#include "CoreMinimal.h"
+#include "FbxHelper.h"
+#include "FbxInclude.h"
+
+/** Forward declarations */
+class FSkeletalMeshAttributes;
+class UInterchangeBaseNode;
+class UInterchangeBaseNodeContainer;
+class UInterchangeSkeletalMeshNode;
+class UInterchangeSkeletalMeshLodDataNode;
+
+struct FMeshDescription;
+
+
+namespace UE
+{
+	namespace Interchange
+	{
+		namespace Private
+		{
+			class FSkeletalMeshGeometryPayload : public FPayloadContextBase
+			{
+			public:
+				virtual ~FSkeletalMeshGeometryPayload() { SkelMeshNodeArray.Empty(); }
+				virtual FString GetPayloadType() const override { return TEXT("SkeletalMesh-GeometryPayload"); }
+				virtual bool FetchPayloadToFile(const FString& PayloadFilepath, TArray<FString>& JSonErrorMessages) override;
+				TArray<FbxNode*> SkelMeshNodeArray;
+				TArray<FbxNode*> SortedJoints;
+				FbxScene* SDKScene = nullptr;
+				FbxGeometryConverter* SDKGeometryConverter = nullptr;
+
+			private:
+				bool AddFbxMeshToMeshDescription(int32 MeshIndex, FMeshDescription& SkeletalMeshDescription, FSkeletalMeshAttributes& SkeletalMeshAttribute, TArray<FString>& JSonErrorMessages);
+			};
+
+			class FFbxSkeletalMesh
+			{
+			public:
+				static void FindSkeletalMeshes(FbxScene* Scene, TArray< TArray<FbxNode*> >& outSkelMeshArray, bool bCombineSkeletalMesh, bool bForceFindRigid);
+				static void AddAllSceneSkeletalMeshes(FbxScene* SDKScene, FbxGeometryConverter* SDKGeometryConverter, UInterchangeBaseNodeContainer& NodeContainer, TArray<FString>& JSonErrorMessages, TMap<FString, TSharedPtr<FPayloadContextBase>>& PayloadContexts, const FString& SourceFilename);
+				static UInterchangeSkeletalMeshNode* CreateSkeletalMeshNode(UInterchangeBaseNodeContainer& NodeContainer, const FString& NodeName, const FString& NodeUniqueID, TArray<FString>& JSonErrorMessages);
+				static UInterchangeSkeletalMeshLodDataNode* CreateSkeletalMeshLodDataNode(UInterchangeBaseNodeContainer& NodeContainer, const FString& NodeName, const FString& NodeUniqueID, TArray<FString>& JSonErrorMessages);
+			
+			private:
+			};
+		}//ns Private
+	}//ns Interchange
+}//ns UE
diff --git a/Engine/Plugins/Runtime/Interchange/Source/Parsers/Fbx/Private/FbxSkeleton.cpp b/Engine/Plugins/Runtime/Interchange/Source/Parsers/Fbx/Private/FbxSkeleton.cpp
new file mode 100644
--- /dev/null
+++ b/Engine/Plugins/Runtime/Interchange/Source/Parsers/Fbx/Private/FbxSkeleton.cpp
@@ -0,0 +1,208 @@
+// Copyright Epic Games, Inc. All Rights Reserved.
+
+#include "FbxSkeleton.h"
+
+#include "CoreMinimal.h"
+#include "FbxConvert.h"
+#include "FbxHelper.h"
+#include "FbxInclude.h"
+#include "FbxSkeletalMesh.h"
+#include "InterchangeJointNode.h"
+#include "InterchangeSkeletonNode.h"
+#include "Nodes/InterchangeBaseNodeContainer.h"
+
+#define GeneratedLODNameSuffix "_GeneratedLOD_"
+
+namespace UE
+{
+	namespace Interchange
+	{
+		namespace Private
+		{
+			void FFbxSkeleton::AddAllSceneSkeletons(FbxScene* SDKScene, UInterchangeBaseNodeContainer& NodeContainer, TArray<FString>& JSonErrorMessages, const FString& SourceFilename)
+			{
+				TArray<TArray<FbxNode*>> SkeletalMeshes;
+				const bool bCombineSkeletalMesh = false;
+				const bool bForceFindRigid = false;
+				FFbxHelper::FindSkeletalMeshes(SDKScene, SkeletalMeshes, bCombineSkeletalMesh, bForceFindRigid);
+				int32 SkeletalMeshCount = SkeletalMeshes.Num();
+				for (int32 SkeletalMeshIndex = 0; SkeletalMeshIndex < SkeletalMeshCount; ++SkeletalMeshIndex)
+				{
+					TArray<FbxNode*>& NodeArray = SkeletalMeshes[SkeletalMeshIndex];
+					int32 MaxNumberOfLOD = 1;
+					for (int32 NodeIndex = 0; NodeIndex < NodeArray.Num(); NodeIndex++)
+					{
+						FbxNode* Node = NodeArray[NodeIndex];
+						if (Node->GetNodeAttribute() && Node->GetNodeAttribute()->GetAttributeType() == FbxNodeAttribute::eLODGroup)
+						{
+							// get max LODgroup level
+							if (MaxNumberOfLOD < Node->GetChildCount())
+							{
+								MaxNumberOfLOD = Node->GetChildCount();
+							}
+						}
+					}
+
+					bool bOverrideSkeletalMeshName = false;
+					FString OverrideSkeletalMeshName;
+					if (SkeletalMeshCount == 1)
+					{
+						bOverrideSkeletalMeshName = true;
+						OverrideSkeletalMeshName = FPaths::GetBaseFilename(SourceFilename);
+					}
+
+					for (int32 LODIndex = 0; LODIndex < MaxNumberOfLOD; LODIndex++)
+					{
+						TArray<FbxNode*> SkelMeshNodeArray;
+						for (int32 j = 0; j < NodeArray.Num(); j++)
+						{
+							FbxNode* Node = NodeArray[j];
+							if (Node->GetNodeAttribute() && Node->GetNodeAttribute()->GetAttributeType() == FbxNodeAttribute::eLODGroup)
+							{
+								TArray<FbxNode*> NodeInLod;
+								if (Node->GetChildCount() > LODIndex)
+								{
+									FFbxHelper::FindAllLODGroupNode(NodeInLod, Node, LODIndex);
+								}
+								else // in less some LODGroups have less level, use the last level
+								{
+									FFbxHelper::FindAllLODGroupNode(NodeInLod, Node, Node->GetChildCount() - 1);
+								}
+
+								for (FbxNode* MeshNode : NodeInLod)
+								{
+									SkelMeshNodeArray.Add(MeshNode);
+								}
+							}
+							else
+							{
+								SkelMeshNodeArray.Add(Node);
+							}
+						}
+						if (SkelMeshNodeArray.Num() <= 0)
+						{
+							//Nothing to import for this LOD
+							continue;
+						}
+						//Find the skeleton for this skeletal mesh LOD
+						TArray<FbxNode*> SortedJoints;
+						FbxArray<FbxAMatrix> LocalsPerLink;
+						if (!FFbxHelper::FindSkeletonJoints(SDKScene, SkelMeshNodeArray, SortedJoints, LocalsPerLink))
+						{
+							continue;
+						}
+						//A skeleton must have at least one joint
+						if (SortedJoints.Num() <= 0)
+						{
+							continue;
+						}
+						FString SkeletonName = (bOverrideSkeletalMeshName ? OverrideSkeletalMeshName : FFbxHelper::GetFbxObjectName(SkelMeshNodeArray[0])) + TEXT("_Skeleton");
+						if(LODIndex > 0)
+						{
+							SkeletonName += TEXT("_") + FString::FromInt(LODIndex);
+						}
+
+						FString SkeletonUniqueID = TEXT("\\Skeleton\\") + FFbxHelper::GetFbxNodeHierarchyName(SortedJoints[0]);
+
+						UInterchangeBaseNode* ExistingNode = NodeContainer.GetNode(*SkeletonUniqueID);
+						if (ExistingNode)
+						{
+							//The skeleton already exist
+							UInterchangeSkeletonNode* ExistingSkeletonNode = Cast<UInterchangeSkeletonNode>(ExistingNode);
+							//It must be of the correct type
+							if (!ensure(ExistingSkeletonNode))
+							{
+								JSonErrorMessages.Add(TEXT("{\"Msg\" : {\"Type\" : \"Error\",\n\"Msg\" : \"Cannot parse properly a fbx skeleton\"}}"));
+							}
+							//Skip this LOD since the skeleton is already created
+							continue;
+						}
+
+						UInterchangeSkeletonNode* SkeletonNode = CreateSkeletonNode(NodeContainer, SkeletonName, SkeletonUniqueID, JSonErrorMessages);
+						if (!ensure(SkeletonNode))
+						{
+							continue;
+						}
+
+						for (int32 JointIndex = 0; JointIndex < SortedJoints.Num(); ++JointIndex)
+						{
+							if (!ensure(LocalsPerLink.GetCount() > JointIndex))
+							{
+								JSonErrorMessages.Add(TEXT("{\"Msg\" : {\"Type\" : \"Error\",\n\"Msg\" : \"Cannot parse properly a fbx skeleton\"}}"));
+								break;
+							}
+							FbxNode* FbxJointNode = SortedJoints[JointIndex];
+							//Create an interchange joint node
+							FString JointName = FFbxHelper::GetFbxObjectName(FbxJointNode);
+							FString JointUniqueID = TEXT("\\Joint\\") + FFbxHelper::GetFbxNodeHierarchyName(FbxJointNode);
+							UInterchangeJointNode* JointNode = CreateJointNode(NodeContainer, JointName, JointUniqueID, JSonErrorMessages);
+							if (!ensure(JointNode))
+							{
+								break;
+							}
+							JointNode->SetCustomName(*JointName);
+							JointNode->SetCustomLocalTransform(FFbxConvert::ConvertTransform(LocalsPerLink[JointIndex]));
+							if (JointIndex == 0)
+							{
+								SkeletonNode->SetCustomRootJointID(*JointUniqueID);
+							}
+							else
+							{
+								const FbxNode* FbxParentJoint = FbxJointNode->GetParent();
+								FString ParentJointUniqueID = TEXT("\\Joint\\") + FFbxHelper::GetFbxNodeHierarchyName(FbxParentJoint);
+								JointNode->SetParentUID(*ParentJointUniqueID);
+							}
+						}
+						FName SkeletonRootJointIDValidator;
+						if (SkeletonNode->GetCustomRootJointID(SkeletonRootJointIDValidator))
+						{
+							UInterchangeBaseNode* RootNode = NodeContainer.GetNode(SkeletonRootJointIDValidator);
+							FName RootNodeDisplayLabel = RootNode->GetDisplayLabel();
+						}
+					}
+				}
+			}
+
+			UInterchangeSkeletonNode* FFbxSkeleton::CreateSkeletonNode(UInterchangeBaseNodeContainer& NodeContainer, const FString& NodeName, const FString& UniqueID, TArray<FString>& JSonErrorMessages)
+			{
+				FName DisplayLabel = *NodeName;
+				FName NodeUID(*UniqueID);
+
+				UInterchangeBaseNode* ExistingNode = NodeContainer.GetNode(NodeUID);
+				if (ExistingNode)
+				{
+					UInterchangeSkeletonNode* ExistingSkeletonNode = Cast<UInterchangeSkeletonNode>(ExistingNode);
+					return (ensure(ExistingSkeletonNode) ? ExistingSkeletonNode : nullptr);
+				}
+
+				UInterchangeSkeletonNode* SkeletonNode = NewObject<UInterchangeSkeletonNode>(&NodeContainer, NAME_None);
+				if (!ensure(SkeletonNode))
+				{
+					JSonErrorMessages.Add(TEXT("{\"Msg\" : {\"Type\" : \"Error\",\n\"Msg\" : \"Cannot allocate a skeleton node when importing fbx\"}}"));
+					return nullptr;
+				}
+				// Creating a UMaterialInterface
+				SkeletonNode->InitializeSkeletonNode(NodeUID, DisplayLabel, TEXT("Skeleton"));
+				NodeContainer.AddNode(SkeletonNode);
+				return SkeletonNode;
+			}
+
+			UInterchangeJointNode* FFbxSkeleton::CreateJointNode(UInterchangeBaseNodeContainer& NodeContainer, const FString& NodeName, const FString& UniqueID, TArray<FString>& JSonErrorMessages)
+			{
+				FName DisplayLabel = *NodeName;
+				FName NodeUID(*UniqueID);
+				UInterchangeJointNode* JointNode = NewObject<UInterchangeJointNode>(&NodeContainer, NAME_None);
+				if (!ensure(JointNode))
+				{
+					JSonErrorMessages.Add(TEXT("{\"Msg\" : {\"Type\" : \"Error\",\n\"Msg\" : \"Cannot allocate a joint node when importing fbx\"}}"));
+					return nullptr;
+				}
+				// Creating a UMaterialInterface
+				JointNode->InitializeNode(NodeUID, DisplayLabel);
+				NodeContainer.AddNode(JointNode);
+				return JointNode;
+			}
+
+		} //ns Private
+	} //ns Interchange
+}//ns UE
diff --git a/Engine/Plugins/Runtime/Interchange/Source/Parsers/Fbx/Private/FbxSkeleton.h b/Engine/Plugins/Runtime/Interchange/Source/Parsers/Fbx/Private/FbxSkeleton.h
new file mode 100644
--- /dev/null
+++ b/Engine/Plugins/Runtime/Interchange/Source/Parsers/Fbx/Private/FbxSkeleton.h
@@ -0,0 +1,31 @@
+// Copyright Epic Games, Inc. All Rights Reserved.
+
+#pragma once
+
+#include "CoreMinimal.h"
+#include "FbxInclude.h"
+
+/** Forward declarations */
+class UInterchangeBaseNode;
+class UInterchangeBaseNodeContainer;
+class UInterchangeJointNode;
+class UInterchangeSkeletonNode;
+
+namespace UE
+{
+	namespace Interchange
+	{
+		namespace Private
+		{
+			class FFbxSkeleton
+			{
+			public:
+				static void AddAllSceneSkeletons(FbxScene* SDKScene, UInterchangeBaseNodeContainer& NodeContainer, TArray<FString>& JSonErrorMessages, const FString& SourceFilename);
+				static UInterchangeSkeletonNode* CreateSkeletonNode(UInterchangeBaseNodeContainer& NodeContainer, const FString& NodeName, const FString& UniqueID, TArray<FString>& JSonErrorMessages);
+				static UInterchangeJointNode* CreateJointNode(UInterchangeBaseNodeContainer& NodeContainer, const FString& NodeName, const FString& UniqueID, TArray<FString>& JSonErrorMessages);
+
+			private:
+			};
+		}//ns Private
+	}//ns Interchange
+}//ns UE
diff --git a/Engine/Plugins/Runtime/Interchange/Source/Parsers/Fbx/Private/InterchangeFbxParser.cpp b/Engine/Plugins/Runtime/Interchange/Source/Parsers/Fbx/Private/InterchangeFbxParser.cpp
--- a/Engine/Plugins/Runtime/Interchange/Source/Parsers/Fbx/Private/InterchangeFbxParser.cpp
+++ b/Engine/Plugins/Runtime/Interchange/Source/Parsers/Fbx/Private/InterchangeFbxParser.cpp
@@ -47,5 +47,20 @@ namespace UE
 			Container.Get()->SaveToFile(ResultFilepath);
 			JsonLoadMessages.Add(TEXT("{\"Msg\" : {\"Type\" : \"Log\",\n\"Msg\" : \"This is a success!\"}}"));
 		}
+
+		void FInterchangeFbxParser::FetchPayload(const FString& PayloadKey, const FString& ResultFolder)
+		{
+			check(FbxParserPrivate.IsValid());
+
+			FString& PayloadFilepath = ResultPayloads.FindOrAdd(PayloadKey);
+			PayloadFilepath = ResultFolder + TEXT("/") + PayloadKey + TEXT(".payload");
+			if (!FbxParserPrivate->FetchPayloadData(PayloadKey, PayloadFilepath, JsonLoadMessages))
+			{
+				JsonLoadMessages.Add(TEXT("{\"Msg\" : {\"Type\" : \"Error\",\n\"Msg\" : \"Cannot fetch fbx payload data.\"}}"));
+				return;
+			}
+			
+			JsonLoadMessages.Add(TEXT("{\"Msg\" : {\"Type\" : \"Log\",\n\"Msg\" : \"Fetch fbx Payload success!\"}}"));
+		}
 	}//ns Interchange
 }//ns UE
diff --git a/Engine/Plugins/Runtime/Interchange/Source/Parsers/Fbx/Public/InterchangeFbxParser.h b/Engine/Plugins/Runtime/Interchange/Source/Parsers/Fbx/Public/InterchangeFbxParser.h
--- a/Engine/Plugins/Runtime/Interchange/Source/Parsers/Fbx/Public/InterchangeFbxParser.h
+++ b/Engine/Plugins/Runtime/Interchange/Source/Parsers/Fbx/Public/InterchangeFbxParser.h
@@ -27,14 +27,29 @@ namespace UE
 			 */
 			void LoadFbxFile(const FString& Filename, const FString& ResultFolder);
 
-			FString GetResultFilepath() { return ResultFilepath; }
-			TArray<FString> GetJsonLoadMessages() { return JsonLoadMessages; }
+			/**
+			 * Extract payload data from the fbx, the key tell the translator what payload the client ask
+			 * @param - PayloadKey is the key that describe the payload data to extract from the fbx file
+			 * @param - ResultFolder is the folder where we must put any result file
+			 */
+			void FetchPayload(const FString& PayloadKey, const FString& ResultFolder);
+
+			FString GetResultFilepath() const { return ResultFilepath; }
+			FString GetResultPayloadFilepath(const FString& PayloadKey) const
+			{
+				if (const FString* PayloadPtr = ResultPayloads.Find(PayloadKey))
+				{
+					return *PayloadPtr;
+				}
+				return FString();
+			}
+			TArray<FString> GetJsonLoadMessages() const { return JsonLoadMessages; }
 
 		private:
 			TStrongObjectPtr<UInterchangeBaseNodeContainer> Container = nullptr;
 			FString ResultFilepath;
 			TArray<FString> JsonLoadMessages;
-
+			TMap<FString, FString> ResultPayloads;
 			TUniquePtr<UE::Interchange::Private::FFbxParser> FbxParserPrivate;
 		};
 	} // ns Interchange
diff --git a/Engine/Source/Programs/InterchangeWorker/Private/InterchangeWorkerImpl.cpp b/Engine/Source/Programs/InterchangeWorker/Private/InterchangeWorkerImpl.cpp
--- a/Engine/Source/Programs/InterchangeWorker/Private/InterchangeWorkerImpl.cpp
+++ b/Engine/Source/Programs/InterchangeWorker/Private/InterchangeWorkerImpl.cpp
@@ -83,6 +83,8 @@ bool FInterchangeWorkerImpl::Run()
 				UE_CLOG(!bIsRunning, LogInterchangeWorker, Error, TEXT("Worker failure: server lost"));
 			}
 		}
+		//Sleep 0 to avoid using too much cpu
+		FPlatformProcess::Sleep(0.0f);
 	}
 
 	UE_CLOG(!bIsRunning, LogInterchangeWorker, Verbose, TEXT("Worker loop exit..."));
@@ -123,6 +125,7 @@ void FInterchangeWorkerImpl::ProcessCommand(const FRunTaskCommand& RunTaskComman
 	FString JSonResult;
 	TArray<FString> JSonMessages;
 	FJsonLoadSourceCmd LoadSourceCommand;
+	FJsonFetchPayloadCmd FetchPayloadCommand;
 	//Any command FromJson function return true if the Json descibe the command
 	if (LoadSourceCommand.FromJson(JsonToProcess))
 	{
@@ -133,7 +136,15 @@ void FInterchangeWorkerImpl::ProcessCommand(const FRunTaskCommand& RunTaskComman
 			ProcessResult = LoadFbxFile(LoadSourceCommand, JSonResult, JSonMessages);
 		}
 	}
-	//TODO add GetPayload command
+	else if (FetchPayloadCommand.FromJson(JsonToProcess))
+	{
+		//Load file command
+		if (FetchPayloadCommand.GetTranslatorID().Equals(TEXT("FBX"), ESearchCase::IgnoreCase))
+		{
+			//We want to load an FBX file
+			ProcessResult = FetchFbxPayload(FetchPayloadCommand, JSonResult, JSonMessages);
+		}
+	}
 	else
 	{
 		ProcessResult = ETaskState::Unknown;
@@ -163,4 +174,17 @@ ETaskState FInterchangeWorkerImpl::LoadFbxFile(const FJsonLoadSourceCmd& LoadSou
 	OutJSonResult = ResultParser.ToJson();
 	ResultState = ETaskState::ProcessOk;
 	return ResultState;
+}
+
+ETaskState FInterchangeWorkerImpl::FetchFbxPayload(const FJsonFetchPayloadCmd& FetchPayloadCommand, FString& OutJSonResult, TArray<FString>& OutJSonMessages)
+{
+	ETaskState ResultState = ETaskState::Unknown;
+	FString PayloadKey = FetchPayloadCommand.GetPayloadKey();
+	FbxParser.FetchPayload(PayloadKey, ResultFolder);
+	FJsonLoadSourceCmd::JsonResultParser ResultParser;
+	ResultParser.SetResultFilename(FbxParser.GetResultPayloadFilepath(PayloadKey));
+	OutJSonMessages = FbxParser.GetJsonLoadMessages();
+	OutJSonResult = ResultParser.ToJson();
+	ResultState = ETaskState::ProcessOk;
+	return ResultState;
 }
\ No newline at end of file
diff --git a/Engine/Source/Programs/InterchangeWorker/Private/InterchangeWorkerImpl.h b/Engine/Source/Programs/InterchangeWorker/Private/InterchangeWorkerImpl.h
--- a/Engine/Source/Programs/InterchangeWorker/Private/InterchangeWorkerImpl.h
+++ b/Engine/Source/Programs/InterchangeWorker/Private/InterchangeWorkerImpl.h
@@ -23,6 +23,7 @@ private:
 	void ProcessCommand(const UE::Interchange::FRunTaskCommand& TerminateCommand);
 
 	UE::Interchange::ETaskState LoadFbxFile(const UE::Interchange::FJsonLoadSourceCmd& LoadSourceCommand, FString& OutJSonResult, TArray<FString>& OutJSonMessages);
+	UE::Interchange::ETaskState FetchFbxPayload(const UE::Interchange::FJsonFetchPayloadCmd& FetchPayloadCommand, FString& OutJSonResult, TArray<FString>& OutJSonMessages);
 
 private:
 	UE::Interchange::FNetworkClientNode NetworkInterface;
diff --git a/Engine/Source/Runtime/Interchange/Core/Public/InterchangeFactoryBase.h b/Engine/Source/Runtime/Interchange/Core/Public/InterchangeFactoryBase.h
--- a/Engine/Source/Runtime/Interchange/Core/Public/InterchangeFactoryBase.h
+++ b/Engine/Source/Runtime/Interchange/Core/Public/InterchangeFactoryBase.h
@@ -72,6 +72,8 @@ public:
 
 	/**
 	 * Create an empty asset from a Node data, this function will be call in the main thread in same time has we create the package.
+	 * The asset create here must have the internal async flag set because the object can be setup in an asynchronous thread and must be consider like an async object until
+	 * The completion task on the object is finish.
 	 *
 	 * @param Arguments - The structure containing all necessary arguments, see the structure definition for the documentation
 	 * @return the created UObject or nullptr if there is an error, see LOG for any error detail.
diff --git a/Engine/Source/Runtime/Interchange/Core/Public/Types/AttributeStorage.h b/Engine/Source/Runtime/Interchange/Core/Public/Types/AttributeStorage.h
--- a/Engine/Source/Runtime/Interchange/Core/Public/Types/AttributeStorage.h
+++ b/Engine/Source/Runtime/Interchange/Core/Public/Types/AttributeStorage.h
@@ -126,6 +126,11 @@ namespace UE
 				}
 				return Ar;
 			}
+
+			FORCEINLINE FString ToString() const
+			{
+				return Key.ToString();
+			}
 		};
 
 		FORCEINLINE uint32 GetTypeHash(const FAttributeKey& AttributeKey)
@@ -1009,7 +1014,6 @@ namespace UE
 			 * Mutex use when accessing or modifying the storage
 			 */
 			mutable FCriticalSection StorageMutex;
-
 		};
 
 		/************************************************************************/
diff --git a/Engine/Source/Runtime/Interchange/Engine/Private/InterchangeManager.cpp b/Engine/Source/Runtime/Interchange/Engine/Private/InterchangeManager.cpp
--- a/Engine/Source/Runtime/Interchange/Engine/Private/InterchangeManager.cpp
+++ b/Engine/Source/Runtime/Interchange/Engine/Private/InterchangeManager.cpp
@@ -142,6 +142,10 @@ void UE::Interchange::FImportAsyncHelper::CancelAndWaitUntilDoneSynchronously()
 	{
 		TasksToComplete.Append(PipelinePostImportTasks);
 	}
+	if (PreAsyncCompletionTask.GetReference())
+	{
+		TasksToComplete.Add(PreAsyncCompletionTask);
+	}
 	if (CompletionTask.GetReference())
 	{
 		//Completion task will make sure any created asset before canceling will be mark for delete
@@ -572,9 +576,15 @@ UE::Interchange::FAssetImportResultRef UInterchangeManager::ImportAssetAsync(con
 
 		for (int32 GraphPipelineIndex = 0; GraphPipelineIndex < PipelineStack.Num(); ++GraphPipelineIndex)
 		{
-			UClass* PipelineClass = PipelineStack[GraphPipelineIndex].LoadSynchronous();
-			UInterchangePipelineBase* GeneratedPipeline = NewObject<UInterchangePipelineBase>(GetTransientPackage(), PipelineClass, NAME_None, RF_NoFlags);
-			AsyncHelper->Pipelines.Add(GeneratedPipeline);
+			if (PipelineStack[GraphPipelineIndex].IsValid())
+			{
+				UClass* PipelineClass = PipelineStack[GraphPipelineIndex].LoadSynchronous();
+				if (PipelineClass)
+				{
+					UInterchangePipelineBase* GeneratedPipeline = NewObject<UInterchangePipelineBase>(GetTransientPackage(), PipelineClass, NAME_None, RF_NoFlags);
+					AsyncHelper->Pipelines.Add(GeneratedPipeline);
+				}
+			}
 		}
 	}
 	else
diff --git a/Engine/Source/Runtime/Interchange/Engine/Private/Tasks/InterchangeTaskCompletion.cpp b/Engine/Source/Runtime/Interchange/Engine/Private/Tasks/InterchangeTaskCompletion.cpp
--- a/Engine/Source/Runtime/Interchange/Engine/Private/Tasks/InterchangeTaskCompletion.cpp
+++ b/Engine/Source/Runtime/Interchange/Engine/Private/Tasks/InterchangeTaskCompletion.cpp
@@ -13,6 +13,18 @@
 #include "UObject/ObjectMacros.h"
 #include "UObject/WeakObjectPtrTemplates.h"
 
+void UE::Interchange::FTaskPreAsyncCompletion::DoTask(ENamedThreads::Type CurrentThread, const FGraphEventRef& MyCompletionGraphEvent)
+{
+#if INTERCHANGE_TRACE_ASYNCHRONOUS_TASK_ENABLED
+	INTERCHANGE_TRACE_ASYNCHRONOUS_TASK(PreAsyncCompletion)
+#endif
+	TSharedPtr<FImportAsyncHelper, ESPMode::ThreadSafe> AsyncHelper = WeakAsyncHelper.Pin();
+	check(AsyncHelper.IsValid());
+
+	//No need anymore of the translators sources
+	AsyncHelper->ReleaseTranslatorsSource();
+}
+
 void UE::Interchange::FTaskCompletion::DoTask(ENamedThreads::Type CurrentThread, const FGraphEventRef& MyCompletionGraphEvent)
 {
 #if INTERCHANGE_TRACE_ASYNCHRONOUS_TASK_ENABLED
@@ -48,7 +60,8 @@ void UE::Interchange::FTaskCompletion::DoTask(ENamedThreads::Type CurrentThread,
 			}
 
 			//Clear any async flag from the created asset
-			Asset->ClearInternalFlags(EInternalObjectFlags::Async);
+			const EInternalObjectFlags AsyncFlags = EInternalObjectFlags::Async | EInternalObjectFlags::AsyncLoading;
+			Asset->ClearInternalFlags(AsyncFlags);
 			//Make sure the package is dirty
 			Asset->MarkPackageDirty();
 #if WITH_EDITOR
diff --git a/Engine/Source/Runtime/Interchange/Engine/Private/Tasks/InterchangeTaskCompletion.h b/Engine/Source/Runtime/Interchange/Engine/Private/Tasks/InterchangeTaskCompletion.h
--- a/Engine/Source/Runtime/Interchange/Engine/Private/Tasks/InterchangeTaskCompletion.h
+++ b/Engine/Source/Runtime/Interchange/Engine/Private/Tasks/InterchangeTaskCompletion.h
@@ -12,6 +12,34 @@ namespace UE
 {
 	namespace Interchange
 	{
+		class FTaskPreAsyncCompletion
+		{
+		private:
+			UInterchangeManager* InterchangeManager;
+			TWeakPtr<FImportAsyncHelper, ESPMode::ThreadSafe> WeakAsyncHelper;
+		public:
+			FTaskPreAsyncCompletion(UInterchangeManager* InInterchangeManager, TWeakPtr<FImportAsyncHelper, ESPMode::ThreadSafe> InAsyncHelper)
+				: InterchangeManager(InInterchangeManager)
+				, WeakAsyncHelper(InAsyncHelper)
+			{
+			}
+
+			static FORCEINLINE ENamedThreads::Type GetDesiredThread()
+			{
+				return ENamedThreads::AnyBackgroundThreadNormalTask;
+			}
+			static FORCEINLINE ESubsequentsMode::Type GetSubsequentsMode()
+			{
+				return ESubsequentsMode::TrackSubsequents;
+			}
+
+			FORCEINLINE TStatId GetStatId() const
+			{
+				RETURN_QUICK_DECLARE_CYCLE_STAT(FTaskPreAsyncCompletion, STATGROUP_TaskGraphTasks);
+			}
+
+			void DoTask(ENamedThreads::Type CurrentThread, const FGraphEventRef& MyCompletionGraphEvent);
+		};
 
 		class FTaskCompletion
 		{
diff --git a/Engine/Source/Runtime/Interchange/Engine/Private/Tasks/InterchangeTaskCreateAsset.cpp b/Engine/Source/Runtime/Interchange/Engine/Private/Tasks/InterchangeTaskCreateAsset.cpp
--- a/Engine/Source/Runtime/Interchange/Engine/Private/Tasks/InterchangeTaskCreateAsset.cpp
+++ b/Engine/Source/Runtime/Interchange/Engine/Private/Tasks/InterchangeTaskCreateAsset.cpp
@@ -32,14 +32,14 @@ namespace UE
 				const FString BaseFileName = FPaths::GetBaseFilename(SourceData->GetFilename());
 
 				//Set the asset name and the package name
-				if (NodeDisplayName.Equals(BaseFileName) || BaseFileName.IsEmpty())
+				//if (NodeDisplayName.Equals(BaseFileName) || BaseFileName.IsEmpty())
 				{
 					OutAssetName = NodeDisplayName;
 				}
-				else
-				{
-					OutAssetName = BaseFileName.IsEmpty() ? NodeDisplayName : BaseFileName + TEXT("_") + NodeDisplayName;
-				}
+// 				else
+// 				{
+// 					OutAssetName = BaseFileName.IsEmpty() ? NodeDisplayName : BaseFileName + TEXT("_") + NodeDisplayName;
+// 				}
 				OutPackageName = FPaths::Combine(*PackageBasePath, *OutAssetName);
 
 				//Sanitize only the package name
@@ -112,6 +112,12 @@ void UE::Interchange::FTaskCreatePackage::DoTask(ENamedThreads::Type CurrentThre
 		UObject* NodeAsset = Factory->CreateEmptyAsset(CreateAssetParams);
 		if (NodeAsset)
 		{
+			if (!NodeAsset->HasAnyInternalFlags(EInternalObjectFlags::Async))
+			{
+				//Sice the async flag is not set we must be in the game thread
+				ensure(IsInGameThread());
+				NodeAsset->SetInternalFlags(EInternalObjectFlags::Async);
+			}
 			FScopeLock Lock(&AsyncHelper->ImportedAssetsPerSourceIndexLock);
 			TArray<UE::Interchange::FImportAsyncHelper::FImportedAssetInfo>& ImportedInfos = AsyncHelper->ImportedAssetsPerSourceIndex.FindOrAdd(SourceIndex);
 			UE::Interchange::FImportAsyncHelper::FImportedAssetInfo& AssetInfo = ImportedInfos.AddDefaulted_GetRef();
diff --git a/Engine/Source/Runtime/Interchange/Engine/Private/Tasks/InterchangeTaskParsing.cpp b/Engine/Source/Runtime/Interchange/Engine/Private/Tasks/InterchangeTaskParsing.cpp
--- a/Engine/Source/Runtime/Interchange/Engine/Private/Tasks/InterchangeTaskParsing.cpp
+++ b/Engine/Source/Runtime/Interchange/Engine/Private/Tasks/InterchangeTaskParsing.cpp
@@ -136,5 +136,9 @@ void UE::Interchange::FTaskParsing::DoTask(ENamedThreads::Type CurrentThread, co
 		CompletionPrerequistes.Add(TaskDatas[TaskIndex].GraphEventRef);
 	}
 
-	AsyncHelper->CompletionTask = TGraphTask<FTaskCompletion>::CreateTask(&CompletionPrerequistes).ConstructAndDispatchWhenReady(InterchangeManager, WeakAsyncHelper);
+	FGraphEventArray PreAsyncCompletionPrerequistes;
+	AsyncHelper->PreAsyncCompletionTask = TGraphTask<FTaskPreAsyncCompletion>::CreateTask(&CompletionPrerequistes).ConstructAndDispatchWhenReady(InterchangeManager, WeakAsyncHelper);
+	PreAsyncCompletionPrerequistes.Add(AsyncHelper->PreAsyncCompletionTask);
+
+	AsyncHelper->CompletionTask = TGraphTask<FTaskCompletion>::CreateTask(&PreAsyncCompletionPrerequistes).ConstructAndDispatchWhenReady(InterchangeManager, WeakAsyncHelper);
 }
diff --git a/Engine/Source/Runtime/Interchange/Engine/Public/InterchangeManager.h b/Engine/Source/Runtime/Interchange/Engine/Public/InterchangeManager.h
--- a/Engine/Source/Runtime/Interchange/Engine/Public/InterchangeManager.h
+++ b/Engine/Source/Runtime/Interchange/Engine/Public/InterchangeManager.h
@@ -156,6 +156,7 @@ namespace UE
 			TArray<FGraphEventRef> CreatePackageTasks;
 			TArray<FGraphEventRef> CreateAssetTasks;
 
+			FGraphEventRef PreAsyncCompletionTask;
 			FGraphEventRef CompletionTask;
 
 			//Create package map, Key is package name. We cannot create package asynchronously so we have to create a game thread task to do this
diff --git a/Engine/Source/Runtime/SkeletalMeshDescription/Private/SkeletalMeshAttributes.cpp b/Engine/Source/Runtime/SkeletalMeshDescription/Private/SkeletalMeshAttributes.cpp
new file mode 100644
--- /dev/null
+++ b/Engine/Source/Runtime/SkeletalMeshDescription/Private/SkeletalMeshAttributes.cpp
@@ -0,0 +1,22 @@
+// Copyright Epic Games, Inc. All Rights Reserved.
+
+#include "SkeletalMeshAttributes.h"
+
+namespace MeshAttribute
+{
+	const FName Vertex::InfluenceCount("InfluenceCount");
+	const FName Vertex::InfluenceBone("InfluenceBone");
+	const FName Vertex::InfluenceWeight("InfluenceWeight");
+}
+
+
+
+void FSkeletalMeshAttributes::Register()
+{
+	// Add skeletalmesh attribute
+	MeshDescription.VertexAttributes().RegisterAttribute<int32[]>(MeshAttribute::Vertex::InfluenceBone, 1, 0, EMeshAttributeFlags::IndexReference | EMeshAttributeFlags::Mandatory);
+	MeshDescription.VertexAttributes().RegisterAttribute<float[]>(MeshAttribute::Vertex::InfluenceWeight, 1, 0.0f, EMeshAttributeFlags::Mandatory);
+
+	// Call super class
+	FStaticMeshAttributes::Register();
+}
diff --git a/Engine/Source/Runtime/SkeletalMeshDescription/Private/SkeletalMeshDescriptionModule.cpp b/Engine/Source/Runtime/SkeletalMeshDescription/Private/SkeletalMeshDescriptionModule.cpp
new file mode 100644
--- /dev/null
+++ b/Engine/Source/Runtime/SkeletalMeshDescription/Private/SkeletalMeshDescriptionModule.cpp
@@ -0,0 +1,28 @@
+// Copyright Epic Games, Inc. All Rights Reserved.
+
+#include "Modules/ModuleInterface.h"
+#include "Modules/ModuleManager.h"
+
+class FSkeletalMeshDescriptionModule : public IModuleInterface
+{
+public:
+
+	FSkeletalMeshDescriptionModule()
+	{
+	}
+
+	virtual void StartupModule() override
+	{
+		// Register any modular features here
+	}
+
+	virtual void ShutdownModule() override
+	{
+		// Unregister any modular features here
+	}
+
+private:
+
+};
+
+IMPLEMENT_MODULE(FSkeletalMeshDescriptionModule, SkeletalMeshDescription);
diff --git a/Engine/Source/Runtime/SkeletalMeshDescription/Private/SkeletalMeshOperations.cpp b/Engine/Source/Runtime/SkeletalMeshDescription/Private/SkeletalMeshOperations.cpp
new file mode 100644
--- /dev/null
+++ b/Engine/Source/Runtime/SkeletalMeshDescription/Private/SkeletalMeshOperations.cpp
@@ -0,0 +1,14 @@
+// Copyright Epic Games, Inc. All Rights Reserved.
+
+#include "SkeletalMeshOperations.h"
+
+#include "SkeletalMeshAttributes.h"
+
+
+DEFINE_LOG_CATEGORY(LogSkeletalMeshOperations);
+
+#define LOCTEXT_NAMESPACE "SkeletalMeshOperations"
+
+//Add specific skeletal mesh descriptions implementation here
+
+#undef LOCTEXT_NAMESPACE
\ No newline at end of file
diff --git a/Engine/Source/Runtime/SkeletalMeshDescription/Public/SkeletalMeshAttributes.h b/Engine/Source/Runtime/SkeletalMeshDescription/Public/SkeletalMeshAttributes.h
new file mode 100644
--- /dev/null
+++ b/Engine/Source/Runtime/SkeletalMeshDescription/Public/SkeletalMeshAttributes.h
@@ -0,0 +1,52 @@
+// Copyright Epic Games, Inc. All Rights Reserved.
+
+#pragma once
+
+#include "MeshAttributes.h"
+#include "StaticMeshAttributes.h"
+
+//Add any skeletalmesh specific attributes here
+
+namespace MeshAttribute
+{
+ 	namespace Vertex
+ 	{
+		extern SKELETALMESHDESCRIPTION_API const FName InfluenceCount;
+ 		extern SKELETALMESHDESCRIPTION_API const FName InfluenceBone;
+		extern SKELETALMESHDESCRIPTION_API const FName InfluenceWeight;
+ 	}
+}
+
+
+class SKELETALMESHDESCRIPTION_API FSkeletalMeshAttributes : public FStaticMeshAttributes
+{
+public:
+
+	explicit FSkeletalMeshAttributes(FMeshDescription& InMeshDescription)
+		: FStaticMeshAttributes(InMeshDescription)
+	{}
+
+	//TODO support uint16 in mesh description
+	TVertexAttributesRef<TArrayAttribute<int32>> GetVertexInfluenceBones() { return MeshDescription.VertexAttributes().GetAttributesRef<TArrayAttribute<int32>>(MeshAttribute::Vertex::InfluenceBone); }
+	TVertexAttributesConstRef<TArrayAttribute<int32>> GetVertexInfluenceBones() const { return MeshDescription.VertexAttributes().GetAttributesRef<TArrayAttribute<int32>>(MeshAttribute::Vertex::InfluenceBone); }
+
+	TVertexAttributesRef<TArrayAttribute<float>> GetVertexInfluenceWeights() { return MeshDescription.VertexAttributes().GetAttributesRef<TArrayAttribute<float>>(MeshAttribute::Vertex::InfluenceWeight); }
+	TVertexAttributesConstRef<TArrayAttribute<float>> GetVertexInfluenceWeights() const { return MeshDescription.VertexAttributes().GetAttributesRef<TArrayAttribute<float>>(MeshAttribute::Vertex::InfluenceWeight); }
+
+	virtual void Register() override;
+
+
+};
+
+
+class FSkeletalMeshConstAttributes : public FStaticMeshConstAttributes
+{
+public:
+
+	explicit FSkeletalMeshConstAttributes(const FMeshDescription& InMeshDescription)
+		: FStaticMeshConstAttributes(InMeshDescription)
+	{}
+
+	TVertexAttributesConstRef<TArrayAttribute<int32>> GetVertexInfluenceBones() const { return MeshDescription.VertexAttributes().GetAttributesRef<TArrayAttribute<int32>>(MeshAttribute::Vertex::InfluenceBone); }
+	TVertexAttributesConstRef<TArrayAttribute<float>> GetVertexInfluenceWeights() const { return MeshDescription.VertexAttributes().GetAttributesRef<TArrayAttribute<float>>(MeshAttribute::Vertex::InfluenceWeight); }
+};
diff --git a/Engine/Source/Runtime/SkeletalMeshDescription/Public/SkeletalMeshOperations.h b/Engine/Source/Runtime/SkeletalMeshDescription/Public/SkeletalMeshOperations.h
new file mode 100644
--- /dev/null
+++ b/Engine/Source/Runtime/SkeletalMeshDescription/Public/SkeletalMeshOperations.h
@@ -0,0 +1,18 @@
+// Copyright Epic Games, Inc. All Rights Reserved.
+
+#pragma once
+
+#include "CoreMinimal.h"
+#include "MeshTypes.h"
+#include "StaticMeshOperations.h"
+
+
+DECLARE_LOG_CATEGORY_EXTERN(LogSkeletalMeshOperations, Log, All);
+
+
+
+class SKELETALMESHDESCRIPTION_API FSkeletalMeshOperations : public FStaticMeshOperations
+{
+public:
+	//Add any specific skeletal mesh operation here
+};
diff --git a/Engine/Source/Runtime/SkeletalMeshDescription/SkeletalMeshDescription.Build.cs b/Engine/Source/Runtime/SkeletalMeshDescription/SkeletalMeshDescription.Build.cs
new file mode 100644
--- /dev/null
+++ b/Engine/Source/Runtime/SkeletalMeshDescription/SkeletalMeshDescription.Build.cs
@@ -0,0 +1,32 @@
+// Copyright Epic Games, Inc. All Rights Reserved.
+
+namespace UnrealBuildTool.Rules
+{
+	public class SkeletalMeshDescription : ModuleRules
+	{
+		public SkeletalMeshDescription(ReadOnlyTargetRules Target) : base(Target)
+		{
+            PrivateIncludePaths.Add("Runtime/SkeletalMeshDescription/Private");
+            PublicIncludePaths.Add("Runtime/SkeletalMeshDescription/Public");
+
+			PublicDependencyModuleNames.AddRange(
+				new string[]
+				{
+					"Core",
+					"CoreUObject",
+					"MeshDescription",
+					"StaticMeshDescription"
+				}
+			);
+
+			PrivateDependencyModuleNames.AddRange(
+				new string[]
+				{
+					"MeshUtilitiesCommon"
+				}
+			);
+
+			AddEngineThirdPartyPrivateStaticDependencies(Target, "MikkTSpace");
+		}
+	}
+}
