Fix CIS build, I remove the material payload data files to fix the CIS, but we will need to put them back when we will need to extract the material payload with interchange framework

diff --git a/Engine/Plugins/Runtime/Interchange/Source/Import/Private/Fbx/InterchangeFbxTranslator.cpp b/Engine/Plugins/Runtime/Interchange/Source/Import/Private/Fbx/InterchangeFbxTranslator.cpp
--- a/Engine/Plugins/Runtime/Interchange/Source/Import/Private/Fbx/InterchangeFbxTranslator.cpp
+++ b/Engine/Plugins/Runtime/Interchange/Source/Import/Private/Fbx/InterchangeFbxTranslator.cpp
@@ -8,7 +8,6 @@
 #include "InterchangeDispatcherTask.h"
 #include "InterchangeManager.h"
 #include "LogInterchangeImportPlugin.h"
-#include "Material/InterchangeMaterialPayloadData.h"
 #include "Mesh/InterchangeSkeletalMeshPayload.h"
 #include "Mesh/InterchangeStaticMeshPayload.h"
 #include "Misc/ConfigCacheIni.h"
@@ -147,12 +146,6 @@ TOptional<UE::Interchange::FImportImage> UInterchangeFbxTranslator::GetTexturePa
 	return TextureTranslator->GetTexturePayloadData(PayloadSourceData, PayLoadKey);
 }
 
-TOptional<UE::Interchange::FMaterialPayloadData> UInterchangeFbxTranslator::GetMaterialPayloadData(const UInterchangeSourceData* SourceData, const FString& PayLoadKey) const
-{
-	//Not implemented, currently we do not have any payload data for the material
-	return TOptional<UE::Interchange::FMaterialPayloadData>();
-}
-
 TOptional<UE::Interchange::FStaticMeshPayloadData> UInterchangeFbxTranslator::GetStaticMeshPayloadData(const FString& PayLoadKey) const
 {
 	//Not implemented, currently we do not have any payload data for static meshes
diff --git a/Engine/Plugins/Runtime/Interchange/Source/Import/Public/Fbx/InterchangeFbxTranslator.h b/Engine/Plugins/Runtime/Interchange/Source/Import/Public/Fbx/InterchangeFbxTranslator.h
--- a/Engine/Plugins/Runtime/Interchange/Source/Import/Public/Fbx/InterchangeFbxTranslator.h
+++ b/Engine/Plugins/Runtime/Interchange/Source/Import/Public/Fbx/InterchangeFbxTranslator.h
@@ -4,7 +4,6 @@
 #include "CoreMinimal.h"
 #include "InterchangeTranslatorBase.h"
 #include "InterchangeDispatcher.h"
-#include "Material/InterchangeMaterialPayloadInterface.h"
 #include "Mesh/InterchangeSkeletalMeshPayload.h"
 #include "Mesh/InterchangeSkeletalMeshPayloadInterface.h"
 #include "Mesh/InterchangeStaticMeshPayload.h"
@@ -21,7 +20,6 @@
 UCLASS(BlueprintType)
 class INTERCHANGEIMPORTPLUGIN_API UInterchangeFbxTranslator : public UInterchangeTranslatorBase
 , public IInterchangeTexturePayloadInterface
-, public IInterchangeMaterialPayloadInterface
 , public IInterchangeStaticMeshPayloadInterface
 , public IInterchangeSkeletalMeshPayloadInterface
 {
@@ -53,22 +51,6 @@ public:
 	/* IInterchangeTexturePayloadInterface End */
 
 
-	//////////////////////////////////////////////////////////////////////////
-	/* IInterchangeMaterialPayloadInterface Begin */
-
-	/**
-	 * Once the translation is done, the import process need a way to retrieve payload data.
-	 * This payload will be use by the factories to create the asset.
-	 *
-	 * @param SourceData - The source data containing the data to translate
-	 * @param PayloadKey - The key to retrieve the a particular payload contain into the specified source data.
-	 * @return a PayloadData containing the imported data. The TOptional will not be set if there is an error.
-	 */
-	virtual TOptional<UE::Interchange::FMaterialPayloadData> GetMaterialPayloadData(const UInterchangeSourceData* SourceData, const FString& PayLoadKey) const override;
-
-	/* IInterchangeMaterialPayloadInterface End */
-
-
 	//////////////////////////////////////////////////////////////////////////
 	/* IInterchangeStaticMeshPayloadInterface Begin */
 
diff --git a/Engine/Plugins/Runtime/Interchange/Source/Import/Public/Material/InterchangeMaterialPayloadData.h b/Engine/Plugins/Runtime/Interchange/Source/Import/Public/Material/InterchangeMaterialPayloadData.h
deleted file mode 100644
diff --git a/Engine/Plugins/Runtime/Interchange/Source/Import/Public/Material/InterchangeMaterialPayloadInterface.h b/Engine/Plugins/Runtime/Interchange/Source/Import/Public/Material/InterchangeMaterialPayloadInterface.h
deleted file mode 100644
