Use FString instead of FName for the attribute storage key

diff --git a/Engine/Plugins/Enterprise/DatasmithImporter/Source/DatasmithImporter/Private/DatasmithTextureImporter.cpp b/Engine/Plugins/Enterprise/DatasmithImporter/Source/DatasmithImporter/Private/DatasmithTextureImporter.cpp
--- a/Engine/Plugins/Enterprise/DatasmithImporter/Source/DatasmithImporter/Private/DatasmithTextureImporter.cpp
+++ b/Engine/Plugins/Enterprise/DatasmithImporter/Source/DatasmithImporter/Private/DatasmithTextureImporter.cpp
@@ -387,7 +387,7 @@ bool UDatasmithTexturePipeline::ExecutePreImportPipeline(UInterchangeBaseNodeCon
 		return false;
 	}
 
-	TArray< FName > Nodes;
+	TArray< FString > Nodes;
 	BaseNodeContainer->GetRoots( Nodes );
 
 	if ( Nodes.Num() <= 0)
diff --git a/Engine/Plugins/Runtime/Interchange/Source/Import/Private/Material/InterchangeMaterialFactory.cpp b/Engine/Plugins/Runtime/Interchange/Source/Import/Private/Material/InterchangeMaterialFactory.cpp
--- a/Engine/Plugins/Runtime/Interchange/Source/Import/Private/Material/InterchangeMaterialFactory.cpp
+++ b/Engine/Plugins/Runtime/Interchange/Source/Import/Private/Material/InterchangeMaterialFactory.cpp
@@ -132,7 +132,7 @@ UObject* UInterchangeMaterialFactory::CreateAsset(const UInterchangeMaterialFact
 			{
 				auto ApplyInputParameter = [&MaterialNode, &Material, &Arguments](EInterchangeMaterialNodeParameterName ParameterName, FExpressionInput& MaterialInput, const FVector2D& Location)
 				{
-					FName OutTextureUID = NAME_None;
+					FString OutTextureUID;
 					int32 OutUVSetIndex = 0;
 					float ScaleU = 0.0f;
 					float ScaleV = 0.0f;
@@ -143,8 +143,8 @@ UObject* UInterchangeMaterialFactory::CreateAsset(const UInterchangeMaterialFact
 						UTexture* TextureReference = nullptr;
 						if (Arguments.NodeContainer)
 						{
-							FName NodeUniqueID = MaterialNode->GetUniqueID();
-							TArray<FName> MaterialDependencies;
+							FString NodeUniqueID = MaterialNode->GetUniqueID();
+							TArray<FString> MaterialDependencies;
 							MaterialNode->GetDependecies(MaterialDependencies);
 							int32 DependenciesCount = MaterialDependencies.Num();
 							for (int32 DependIndex = 0; DependIndex < DependenciesCount; ++DependIndex)
diff --git a/Engine/Plugins/Runtime/Interchange/Source/Import/Private/Mesh/InterchangeSkeletalMeshFactory.cpp b/Engine/Plugins/Runtime/Interchange/Source/Import/Private/Mesh/InterchangeSkeletalMeshFactory.cpp
--- a/Engine/Plugins/Runtime/Interchange/Source/Import/Private/Mesh/InterchangeSkeletalMeshFactory.cpp
+++ b/Engine/Plugins/Runtime/Interchange/Source/Import/Private/Mesh/InterchangeSkeletalMeshFactory.cpp
@@ -37,12 +37,12 @@ namespace UE
 		{
 			struct FJointInfo
 			{
-				FName Name;
+				FString Name;
 				int32 ParentIndex;  // 0 if this is the root bone.  
 				FTransform	LocalTransform; // local transform
 			};
 
-			void RecursiveAddBones(const UInterchangeBaseNodeContainer* NodeContainer, FName JointNodeId, TArray <FJointInfo>& JointInfos, int32 ParentIndex, TArray<SkeletalMeshImportData::FBone>& RefBonesBinary)
+			void RecursiveAddBones(const UInterchangeBaseNodeContainer* NodeContainer, const FString& JointNodeId, TArray <FJointInfo>& JointInfos, int32 ParentIndex, TArray<SkeletalMeshImportData::FBone>& RefBonesBinary)
 			{
 				const UInterchangeJointNode* JointNode = Cast<UInterchangeJointNode>(NodeContainer->GetNode(JointNodeId));
 				if (!JointNode)
@@ -58,7 +58,7 @@ namespace UE
 				Info.ParentIndex = ParentIndex;
 
 				SkeletalMeshImportData::FBone& Bone = RefBonesBinary.AddZeroed_GetRef();
-				Bone.Name = Info.Name.ToString();
+				Bone.Name = Info.Name;
 				Bone.BonePos.Transform = Info.LocalTransform;
 				Bone.ParentIndex = ParentIndex;
 				//Fill the scrap we do not need
@@ -67,7 +67,7 @@ namespace UE
 				Bone.BonePos.YSize = 1.0f;
 				Bone.BonePos.ZSize = 1.0f;
 				
-				const TArray<FName> ChildrenIds = NodeContainer->GetNodeChildrenUIDs(JointNodeId);
+				const TArray<FString> ChildrenIds = NodeContainer->GetNodeChildrenUIDs(JointNodeId);
 				Bone.NumChildren = ChildrenIds.Num();
 				for (int32 ChildIndex = 0; ChildIndex < ChildrenIds.Num(); ++ChildIndex)
 				{
@@ -75,7 +75,7 @@ namespace UE
 				}
 			}
 
-			bool ProcessImportMeshSkeleton(const USkeleton* SkeletonAsset, FReferenceSkeleton& RefSkeleton, int32& SkeletalDepth, const UInterchangeBaseNodeContainer* NodeContainer, FName RootJointNodeId, TArray<SkeletalMeshImportData::FBone>& RefBonesBinary)
+			bool ProcessImportMeshSkeleton(const USkeleton* SkeletonAsset, FReferenceSkeleton& RefSkeleton, int32& SkeletalDepth, const UInterchangeBaseNodeContainer* NodeContainer, const FString& RootJointNodeId, TArray<SkeletalMeshImportData::FBone>& RefBonesBinary)
 			{
 				auto FixupBoneName = [](FString BoneName)
 				{
@@ -96,8 +96,8 @@ namespace UE
 				{
 					const FJointInfo& BinaryBone = JointInfos[b];
 
-					const FString BoneName = FixupBoneName(BinaryBone.Name.ToString());
-					const FMeshBoneInfo BoneInfo(FName(*BoneName, FNAME_Add), BinaryBone.Name.ToString(), BinaryBone.ParentIndex);
+					const FString BoneName = FixupBoneName(BinaryBone.Name);
+					const FMeshBoneInfo BoneInfo(FName(*BoneName, FNAME_Add), BinaryBone.Name, BinaryBone.ParentIndex);
 					const FTransform BoneTransform(BinaryBone.LocalTransform);
 					if (RefSkeleton.FindRawBoneIndex(BoneInfo.Name) != INDEX_NONE)
 					{
@@ -702,11 +702,11 @@ namespace UE
 				return true;
 			}
 
-			bool CopyBlendShapesMeshDescriptionToSkeletalMeshImportData(const TArray<FName>& BlendShapeToImport, const TMap<FString, FMeshDescription>& LodBlendShapeMeshDescriptions, FSkeletalMeshImportData& DestinationSkeletalMeshImportData)
+			bool CopyBlendShapesMeshDescriptionToSkeletalMeshImportData(const TArray<FString>& BlendShapeToImport, const TMap<FString, FMeshDescription>& LodBlendShapeMeshDescriptions, FSkeletalMeshImportData& DestinationSkeletalMeshImportData)
 			{
 				for (const TPair<FString, FMeshDescription>& Pair : LodBlendShapeMeshDescriptions)
 				{
-					FName BlendShapeName(*(Pair.Key));
+					FString BlendShapeName(Pair.Key);
 					//Skip blend shape that is not define in the BlendShapeToImport array
 					if (!BlendShapeToImport.Contains(BlendShapeName))
 					{
@@ -754,7 +754,7 @@ namespace UE
 			 * Fill the Materials array using the raw skeletalmesh geometry data (using material imported name)
 			 * Find the material from the dependencies of the skeletalmesh before searching in all package.
 			 */
-			void ProcessImportMeshMaterials(TArray<FSkeletalMaterial>& Materials, FSkeletalMeshImportData& ImportData, TMap<FName, UMaterialInterface*>& AvailableMaterials)
+			void ProcessImportMeshMaterials(TArray<FSkeletalMaterial>& Materials, FSkeletalMeshImportData& ImportData, TMap<FString, UMaterialInterface*>& AvailableMaterials)
 			{
 				TArray <SkeletalMeshImportData::FMaterial>& ImportedMaterials = ImportData.Materials;
 				// If direct linkup of materials is requested, try to find them here - to get a texture name from a 
@@ -794,15 +794,14 @@ namespace UE
 					if(!Material)
 					{
 						//Try to find the material in the skeletal mesh node dependencies (Materials are import before skeletal mesh when there is a dependency)
-						const FName MaterialFName = FName(*ImportedMaterial.MaterialImportName);
-						if (AvailableMaterials.Contains(FName(MaterialFName)))
+						const FString MaterialName = ImportedMaterial.MaterialImportName;
+						if (AvailableMaterials.Contains(MaterialName))
 						{
-							Material = AvailableMaterials.FindChecked(MaterialFName);
+							Material = AvailableMaterials.FindChecked(MaterialName);
 						}
 						else
 						{
 							//We did not found any material in the dependencies so try to find material everywhere
-							FString MaterialName = ImportedMaterial.MaterialImportName;
 							Material = FindObject<UMaterialInterface>(ANY_PACKAGE, *MaterialName);
 							if (Material == nullptr && MaterialNameNoSkin.Len() < MaterialName.Len())
 							{
@@ -1119,20 +1118,20 @@ UObject* UInterchangeSkeletalMeshFactory::CreateAsset(const UInterchangeSkeletal
 			}
 			USkeleton* SkeletonReference = nullptr;
 			int32 LodCount = SkeletalMeshNode->GetLodDataCount();
-			TArray<FName> LodDataUniqueIds;
+			TArray<FString> LodDataUniqueIds;
 			SkeletalMeshNode->GetLodDataUniqueIds(LodDataUniqueIds);
 			ensure(LodDataUniqueIds.Num() == LodCount);
 			int32 CurrentLodIndex = 0;
 			for (int32 LodIndex = 0; LodIndex < LodCount; ++LodIndex)
 			{
-				FName LodUniqueId = LodDataUniqueIds[LodIndex];
+				FString LodUniqueId = LodDataUniqueIds[LodIndex];
 				const UInterchangeSkeletalMeshLodDataNode* LodDataNode = Cast<UInterchangeSkeletalMeshLodDataNode>(Arguments.NodeContainer->GetNode(LodUniqueId));
 				if (!LodDataNode)
 				{
 					UE_LOG(LogInterchangeImportPlugin, Warning, TEXT("Invalid LOD when importing SkeletalMesh asset %s"), *Arguments.AssetName);
 					continue;
 				}
-				FName SkeletonNodeId = NAME_None;
+				FString SkeletonNodeId;
 				if (!LodDataNode->GetCustomSkeletonID(SkeletonNodeId))
 				{
 					UE_LOG(LogInterchangeImportPlugin, Warning, TEXT("Invalid Skeleton LOD when importing SkeletalMesh asset %s"), *Arguments.AssetName);
@@ -1159,7 +1158,7 @@ UObject* UInterchangeSkeletalMeshFactory::CreateAsset(const UInterchangeSkeletal
 					}
 				}
 
-				FName RootJointNodeId = NAME_None;
+				FString RootJointNodeId;
 				if (!SkeletonNode->GetCustomRootJointID(RootJointNodeId))
 				{
 					UE_LOG(LogInterchangeImportPlugin, Warning, TEXT("Invalid Skeleton LOD Root Joint when importing SkeletalMesh asset %s"), *Arguments.AssetName);
@@ -1172,7 +1171,7 @@ UObject* UInterchangeSkeletalMeshFactory::CreateAsset(const UInterchangeSkeletal
 
 				//Add the lod mesh data to the skeletalmesh
 				FSkeletalMeshImportData SkeletalMeshImportData;
-				TArray<FName> TranslatorMeshKeys;
+				TArray<FString> TranslatorMeshKeys;
 				LodDataNode->GetTranslatorMeshKeys(TranslatorMeshKeys);
 				//We should have only one key per LOD
 				ensure(TranslatorMeshKeys.Num() == 1);
@@ -1181,10 +1180,10 @@ UObject* UInterchangeSkeletalMeshFactory::CreateAsset(const UInterchangeSkeletal
 				{
 					int32 MeshLodIndex = 0;
 
-					TOptional<UE::Interchange::FSkeletalMeshLodPayloadData> LodMeshPayload = SkeletalMeshTranslatorPayloadInterface->GetSkeletalMeshLodPayloadData(TranslatorMeshKeys[MeshLodIndex].ToString());
+					TOptional<UE::Interchange::FSkeletalMeshLodPayloadData> LodMeshPayload = SkeletalMeshTranslatorPayloadInterface->GetSkeletalMeshLodPayloadData(TranslatorMeshKeys[MeshLodIndex]);
 					if (!LodMeshPayload.IsSet())
 					{
-						UE_LOG(LogInterchangeImportPlugin, Warning, TEXT("Invalid Skeletal mesh payload key [%s] SkeletalMesh asset %s"), *TranslatorMeshKeys[MeshLodIndex].ToString(), *Arguments.AssetName);
+						UE_LOG(LogInterchangeImportPlugin, Warning, TEXT("Invalid Skeletal mesh payload key [%s] SkeletalMesh asset %s"), *TranslatorMeshKeys[MeshLodIndex], *Arguments.AssetName);
 						break;
 					}
 					//TODO use the mesh description as the source import data for skeletalmesh
@@ -1196,7 +1195,7 @@ UObject* UInterchangeSkeletalMeshFactory::CreateAsset(const UInterchangeSkeletal
 					SkeletalMeshImportData.RefBonesBinary = RefBonesBinary;
 
 					//Pipeline can remove names from the list to control which blendshapes they import
-					TArray<FName> BlendShapeToImport;
+					TArray<FString> BlendShapeToImport;
 					LodDataNode->GetBlendShapes(BlendShapeToImport);
 					//Copy also the blend shapes data
 					UE::Interchange::Private::CopyBlendShapesMeshDescriptionToSkeletalMeshImportData(BlendShapeToImport, LodMeshPayload->LodBlendShapeMeshDescriptions, SkeletalMeshImportData);
@@ -1205,8 +1204,8 @@ UObject* UInterchangeSkeletalMeshFactory::CreateAsset(const UInterchangeSkeletal
 				ensure(ImportedResource->LODModels.Add(new FSkeletalMeshLODModel()) == CurrentLodIndex);
 				FSkeletalMeshLODModel& LODModel = ImportedResource->LODModels[CurrentLodIndex];
 
-				TMap<FName, UMaterialInterface*> AvailableMaterials;
-				TArray<FName> Dependencies;
+				TMap<FString, UMaterialInterface*> AvailableMaterials;
+				TArray<FString> Dependencies;
 				SkeletalMeshNode->GetDependecies(Dependencies);
 				for (int32 DependencyIndex = 0; DependencyIndex < Dependencies.Num(); ++DependencyIndex)
 				{
diff --git a/Engine/Plugins/Runtime/Interchange/Source/Import/Private/Texture/InterchangeTextureFactory.cpp b/Engine/Plugins/Runtime/Interchange/Source/Import/Private/Texture/InterchangeTextureFactory.cpp
--- a/Engine/Plugins/Runtime/Interchange/Source/Import/Private/Texture/InterchangeTextureFactory.cpp
+++ b/Engine/Plugins/Runtime/Interchange/Source/Import/Private/Texture/InterchangeTextureFactory.cpp
@@ -181,7 +181,7 @@ UObject* UInterchangeTextureFactory::CreateAsset(const UInterchangeTextureFactor
 			UInterchangeBaseNode* PreviousNode = nullptr;
 			if (InterchangeAssetImportData)
 			{
-				PreviousNode = InterchangeAssetImportData->NodeContainer->GetNode(FName(*InterchangeAssetImportData->NodeUniqueID));
+				PreviousNode = InterchangeAssetImportData->NodeContainer->GetNode(InterchangeAssetImportData->NodeUniqueID);
 			}
 			UInterchangeTextureNode* CurrentNode = NewObject<UInterchangeTextureNode>();
 			UInterchangeBaseNode::CopyStorage(TextureNode, CurrentNode);
diff --git a/Engine/Plugins/Runtime/Interchange/Source/Import/Private/Texture/TextureTranslatorUtilities.cpp b/Engine/Plugins/Runtime/Interchange/Source/Import/Private/Texture/TextureTranslatorUtilities.cpp
--- a/Engine/Plugins/Runtime/Interchange/Source/Import/Private/Texture/TextureTranslatorUtilities.cpp
+++ b/Engine/Plugins/Runtime/Interchange/Source/Import/Private/Texture/TextureTranslatorUtilities.cpp
@@ -17,8 +17,8 @@ bool UE::Interchange::FTextureTranslatorUtilities::Generic2DTextureTranslate(con
 		return false;
 	}
 
-	FName DisplayLabel = *FPaths::GetBaseFilename(Filename);
-	FName NodeUID(*Filename);
+	FString DisplayLabel = FPaths::GetBaseFilename(Filename);
+	FString NodeUID(Filename);
 	UInterchangeTextureNode* TextureNode = NewObject<UInterchangeTextureNode>(&BaseNodeContainer, NAME_None);
 	if (!ensure(TextureNode))
 	{
diff --git a/Engine/Plugins/Runtime/Interchange/Source/Nodes/Public/InterchangeJointNode.h b/Engine/Plugins/Runtime/Interchange/Source/Nodes/Public/InterchangeJointNode.h
--- a/Engine/Plugins/Runtime/Interchange/Source/Nodes/Public/InterchangeJointNode.h
+++ b/Engine/Plugins/Runtime/Interchange/Source/Nodes/Public/InterchangeJointNode.h
@@ -49,15 +49,15 @@ public:
 	//Joint name attribute, all joints in a skeleton must have unique name
 
 	UFUNCTION(BlueprintCallable, Category = "Interchange | Node | Joint")
-	bool GetCustomName(FName& AttributeValue) const
+	bool GetCustomName(FString& AttributeValue) const
 	{
-		IMPLEMENT_NODE_ATTRIBUTE_GETTER(Name, FName);
+		IMPLEMENT_NODE_ATTRIBUTE_GETTER(Name, FString);
 	}
 
 	UFUNCTION(BlueprintCallable, Category = "Interchange | Node | Joint")
-	bool SetCustomName(const FName& AttributeValue)
+	bool SetCustomName(const FString& AttributeValue)
 	{
-		IMPLEMENT_NODE_ATTRIBUTE_SETTER_NODELEGATE(Name, FName)
+		IMPLEMENT_NODE_ATTRIBUTE_SETTER_NODELEGATE(Name, FString)
 	}
 
 	//Local transform attribute
diff --git a/Engine/Plugins/Runtime/Interchange/Source/Nodes/Public/InterchangeMaterialNode.h b/Engine/Plugins/Runtime/Interchange/Source/Nodes/Public/InterchangeMaterialNode.h
--- a/Engine/Plugins/Runtime/Interchange/Source/Nodes/Public/InterchangeMaterialNode.h
+++ b/Engine/Plugins/Runtime/Interchange/Source/Nodes/Public/InterchangeMaterialNode.h
@@ -90,7 +90,7 @@ public:
 	 * @param InAssetClass - The class the material factory will create for this node.
 	 */
 	UFUNCTION(BlueprintCallable, Category = "Interchange | Node | Material")
-	void InitializeMaterialNode(const FName& UniqueID, const FName& DisplayLabel, const FString& InAssetClass)
+	void InitializeMaterialNode(const FString& UniqueID, const FString& DisplayLabel, const FString& InAssetClass)
 	{
 		bIsMaterialNodeClassInitialized = false;
 		InitializeNode(UniqueID, DisplayLabel);
@@ -199,7 +199,7 @@ public:
 	 * @note - A parameter name can have only one of the 3 type of input set, the last input type set is the one that will be created.
 	 */
 	UFUNCTION(BlueprintCallable, Category = "Interchange | Node | Material")
-	void AddTextureParameterData(const EInterchangeMaterialNodeParameterName ParameterName, const FName TextureUID, const int32 UVSetIndex, const float ScaleU, const float ScaleV)
+	void AddTextureParameterData(const EInterchangeMaterialNodeParameterName ParameterName, const FString& TextureUID, const int32 UVSetIndex, const float ScaleU, const float ScaleV)
 	{
 		FParameterData& ParameterData = ParameterDatas.FindOrAdd(ParameterName);
 		ParameterData.bIsTextureParameter = true;
@@ -212,7 +212,7 @@ public:
 	}
 
 	UFUNCTION(BlueprintCallable, Category = "Interchange | Node | Material")
-	bool GetTextureParameterData(const EInterchangeMaterialNodeParameterName ParameterName, FName& OutTextureUID, int32& OutUVSetIndex, float& OutScaleU, float& OutScaleV) const
+	bool GetTextureParameterData(const EInterchangeMaterialNodeParameterName ParameterName, FString& OutTextureUID, int32& OutUVSetIndex, float& OutScaleU, float& OutScaleV) const
 	{
 		if (const FParameterData* ParameterData = ParameterDatas.Find(ParameterName))
 		{
@@ -353,7 +353,7 @@ private:
 	struct FParameterData
 	{
 		bool bIsTextureParameter = false;
-		FName TextureUID = NAME_None;
+		FString TextureUID;
 		int32 UVSetIndex = 0; //The "UV set" index we use to set the material sampler input default is channel 0
 		float ScaleU = 1.0f;
 		float ScaleV = 1.0f;
diff --git a/Engine/Plugins/Runtime/Interchange/Source/Nodes/Public/InterchangeSkeletalMeshLodDataNode.h b/Engine/Plugins/Runtime/Interchange/Source/Nodes/Public/InterchangeSkeletalMeshLodDataNode.h
--- a/Engine/Plugins/Runtime/Interchange/Source/Nodes/Public/InterchangeSkeletalMeshLodDataNode.h
+++ b/Engine/Plugins/Runtime/Interchange/Source/Nodes/Public/InterchangeSkeletalMeshLodDataNode.h
@@ -62,15 +62,15 @@ public:
 public:
 	/** Return false if the Attribute was not set previously.*/
 	UFUNCTION(BlueprintCallable, Category = "Interchange | Node | SkeletalMeshLodData")
-	bool GetCustomSkeletonID(FName& AttributeValue) const
+	bool GetCustomSkeletonID(FString& AttributeValue) const
 	{
-		IMPLEMENT_NODE_ATTRIBUTE_GETTER(SkeletonID, FName);
+		IMPLEMENT_NODE_ATTRIBUTE_GETTER(SkeletonID, FString);
 	}
 
 	UFUNCTION(BlueprintCallable, Category = "Interchange | Node | SkeletalMeshLodData")
-	bool SetCustomSkeletonID(const FName& AttributeValue)
+	bool SetCustomSkeletonID(const FString& AttributeValue)
 	{
-		IMPLEMENT_NODE_ATTRIBUTE_SETTER_NODELEGATE(SkeletonID, FName)
+		IMPLEMENT_NODE_ATTRIBUTE_SETTER_NODELEGATE(SkeletonID, FString)
 	}
 
 	/* Translator mesh keys*/
@@ -82,19 +82,19 @@ public:
 	}
 
 	UFUNCTION(BlueprintCallable, Category = "Interchange | Node | SkeletalMeshLodData")
-	void GetTranslatorMeshKeys(TArray<FName>& OutTranslatorMeshKeys) const
+	void GetTranslatorMeshKeys(TArray<FString>& OutTranslatorMeshKeys) const
 	{
 		TranslatorMeshKeys.GetNames(OutTranslatorMeshKeys);
 	}
 
 	UFUNCTION(BlueprintCallable, Category = "Interchange | Node | SkeletalMeshLodData")
-	bool AddTranslatorMeshKey(FName TranslatorMeshKey)
+	bool AddTranslatorMeshKey(const FString& TranslatorMeshKey)
 	{
 		return TranslatorMeshKeys.AddName(TranslatorMeshKey);
 	}
 
 	UFUNCTION(BlueprintCallable, Category = "Interchange | Node | SkeletalMeshLodData")
-	bool RemoveTranslatorMeshKey(FName TranslatorMeshKey)
+	bool RemoveTranslatorMeshKey(const FString& TranslatorMeshKey)
 	{
 		return TranslatorMeshKeys.RemoveName(TranslatorMeshKey);
 	}
@@ -108,19 +108,19 @@ public:
 	}
 
 	UFUNCTION(BlueprintCallable, Category = "Interchange | Node | SkeletalMeshLodData")
-	void GetBlendShapes(TArray<FName>& OutBlendShapeNames) const
+	void GetBlendShapes(TArray<FString>& OutBlendShapeNames) const
 	{
 		BlendShapesNames.GetNames(OutBlendShapeNames);
 	}
 
 	UFUNCTION(BlueprintCallable, Category = "Interchange | Node | SkeletalMeshLodData")
-	bool AddBlendShape(FName BlendShapeName)
+	bool AddBlendShape(const FString& BlendShapeName)
 	{
 		return BlendShapesNames.AddName(BlendShapeName);
 	}
 
 	UFUNCTION(BlueprintCallable, Category = "Interchange | Node | SkeletalMeshLodData")
-	bool RemoveBlendShape(FName BlendShapeName)
+	bool RemoveBlendShape(const FString& BlendShapeName)
 	{
 		return BlendShapesNames.RemoveName(BlendShapeName);
 	}
diff --git a/Engine/Plugins/Runtime/Interchange/Source/Nodes/Public/InterchangeSkeletalMeshNode.h b/Engine/Plugins/Runtime/Interchange/Source/Nodes/Public/InterchangeSkeletalMeshNode.h
--- a/Engine/Plugins/Runtime/Interchange/Source/Nodes/Public/InterchangeSkeletalMeshNode.h
+++ b/Engine/Plugins/Runtime/Interchange/Source/Nodes/Public/InterchangeSkeletalMeshNode.h
@@ -51,7 +51,7 @@ public:
 	 * @param InAssetClass - The class the SkeletalMesh factory will create for this node.
 	 */
 	UFUNCTION(BlueprintCallable, Category = "Interchange | Node | SkeletalMesh")
-	void InitializeSkeletalMeshNode(const FName& UniqueID, const FName& DisplayLabel, const FString& InAssetClass)
+	void InitializeSkeletalMeshNode(const FString& UniqueID, const FString& DisplayLabel, const FString& InAssetClass)
 	{
 		bIsNodeClassInitialized = false;
 		InitializeNode(UniqueID, DisplayLabel);
@@ -109,19 +109,19 @@ public:
 	}
 
 	UFUNCTION(BlueprintCallable, Category = "Interchange | Node | SkeletalMesh")
-	void GetLodDataUniqueIds(TArray<FName>& OutLodDataUniqueIds) const
+	void GetLodDataUniqueIds(TArray<FString>& OutLodDataUniqueIds) const
 	{
 		LodDependencies.GetNames(OutLodDataUniqueIds);
 	}
 
 	UFUNCTION(BlueprintCallable, Category = "Interchange | Node | SkeletalMesh")
-	bool AddLodDataUniqueId(FName LodDataUniqueId)
+	bool AddLodDataUniqueId(const FString& LodDataUniqueId)
 	{
 		return LodDependencies.AddName(LodDataUniqueId);
 	}
 
 	UFUNCTION(BlueprintCallable, Category = "Interchange | Node | SkeletalMesh")
-	bool RemoveLodDataUniqueId(FName LodDataUniqueId)
+	bool RemoveLodDataUniqueId(const FString& LodDataUniqueId)
 	{
 		return LodDependencies.RemoveName(LodDataUniqueId);
 	}
diff --git a/Engine/Plugins/Runtime/Interchange/Source/Nodes/Public/InterchangeSkeletonNode.h b/Engine/Plugins/Runtime/Interchange/Source/Nodes/Public/InterchangeSkeletonNode.h
--- a/Engine/Plugins/Runtime/Interchange/Source/Nodes/Public/InterchangeSkeletonNode.h
+++ b/Engine/Plugins/Runtime/Interchange/Source/Nodes/Public/InterchangeSkeletonNode.h
@@ -32,7 +32,7 @@ public:
 	 * @param InAssetClass - The class the Skeleton factory will create for this node.
 	 */
 	UFUNCTION(BlueprintCallable, Category = "Interchange | Node | Skeleton")
-	void InitializeSkeletonNode(const FName& UniqueID, const FName& DisplayLabel, const FString& InAssetClass)
+	void InitializeSkeletonNode(const FString& UniqueID, const FString& DisplayLabel, const FString& InAssetClass)
 	{
 		bIsNodeClassInitialized = false;
 		InitializeNode(UniqueID, DisplayLabel);
@@ -83,15 +83,15 @@ public:
 public:
 	/** Return false if the Attribute was not set previously.*/
 	UFUNCTION(BlueprintCallable, Category = "Interchange | Node | Skeleton")
-	bool GetCustomRootJointID(FName& AttributeValue) const
+	bool GetCustomRootJointID(FString& AttributeValue) const
 	{
-		IMPLEMENT_NODE_ATTRIBUTE_GETTER(RootJointID, FName);
+		IMPLEMENT_NODE_ATTRIBUTE_GETTER(RootJointID, FString);
 	}
 
 	UFUNCTION(BlueprintCallable, Category = "Interchange | Node | Skeleton")
-	bool SetCustomRootJointID(const FName& AttributeValue)
+	bool SetCustomRootJointID(const FString& AttributeValue)
 	{
-		IMPLEMENT_NODE_ATTRIBUTE_SETTER_NODELEGATE(RootJointID, FName)
+		IMPLEMENT_NODE_ATTRIBUTE_SETTER_NODELEGATE(RootJointID, FString)
 	}
 
 private:
diff --git a/Engine/Plugins/Runtime/Interchange/Source/Nodes/Public/InterchangeTextureNode.h b/Engine/Plugins/Runtime/Interchange/Source/Nodes/Public/InterchangeTextureNode.h
--- a/Engine/Plugins/Runtime/Interchange/Source/Nodes/Public/InterchangeTextureNode.h
+++ b/Engine/Plugins/Runtime/Interchange/Source/Nodes/Public/InterchangeTextureNode.h
@@ -51,7 +51,7 @@ public:
 	 * @param InAssetClass - The class the texture factory will create for this node.
 	 */
 	UFUNCTION(BlueprintCallable, Category = "Interchange | Node | Texture")
-	void InitializeTextureNode(const FName& UniqueID, const FName& DisplayLabel, const FString& InAssetClass)
+	void InitializeTextureNode(const FString& UniqueID, const FString& DisplayLabel, const FString& InAssetClass)
 	{
 		bIsTextureNodeClassInitialized = false;
 		InitializeNode(UniqueID, DisplayLabel);
diff --git a/Engine/Plugins/Runtime/Interchange/Source/Parsers/Fbx/Private/FbxMaterial.cpp b/Engine/Plugins/Runtime/Interchange/Source/Parsers/Fbx/Private/FbxMaterial.cpp
--- a/Engine/Plugins/Runtime/Interchange/Source/Parsers/Fbx/Private/FbxMaterial.cpp
+++ b/Engine/Plugins/Runtime/Interchange/Source/Parsers/Fbx/Private/FbxMaterial.cpp
@@ -18,10 +18,10 @@ namespace UE
 		{
 			UInterchangeMaterialNode* FFbxMaterial::CreateMaterialNode(UInterchangeBaseNodeContainer& NodeContainer, const FString& NodeName, TArray<FString>& JSonErrorMessages)
 			{
-				FName NodeUID(*NodeName);
+				FString NodeUID(NodeName);
 				FString MaterialNameNoSkin = UE::Interchange::Material::RemoveSkinFromName(NodeName);
 
-				FName DisplayLabel = *MaterialNameNoSkin;
+				FString DisplayLabel = MaterialNameNoSkin;
 				UInterchangeMaterialNode* MaterialNode = NewObject<UInterchangeMaterialNode>(&NodeContainer, NAME_None);
 				if (!ensure(MaterialNode))
 				{
@@ -37,8 +37,8 @@ namespace UE
 
 			UInterchangeTextureNode* FFbxMaterial::CreateTexture2DNode(UInterchangeBaseNodeContainer& NodeContainer, const FString& TextureFilePath, TArray<FString>& JSonErrorMessages)
 			{
-				FName DisplayLabel = *FPaths::GetBaseFilename(TextureFilePath);
-				FName NodeUID(*TextureFilePath);
+				FString DisplayLabel = FPaths::GetBaseFilename(TextureFilePath);
+				FString NodeUID(TextureFilePath);
 				UInterchangeTextureNode* TextureNode = NewObject<UInterchangeTextureNode>(&NodeContainer, NAME_None);
 				if (!ensure(TextureNode))
 				{
@@ -56,7 +56,7 @@ namespace UE
 			{
 				//Create a material node
 				FString MaterialName = FFbxHelper::GetFbxObjectName(SurfaceMaterial);
-				FName NodeUID(*MaterialName);
+				FString NodeUID(MaterialName);
 				UInterchangeMaterialNode* MaterialNode = Cast<UInterchangeMaterialNode>(NodeContainer.GetNode(NodeUID));
 				if (!MaterialNode)
 				{
@@ -94,7 +94,7 @@ namespace UE
 									}
 									//Create a texture node and make it child of the material node
 									TArray<FString> JsonErrorMessage;
-									FName NodeUID(*TextureFilename);
+									FString NodeUID(TextureFilename);
 									UInterchangeTextureNode* TextureNode = Cast<UInterchangeTextureNode>(NodeContainer.GetNode(NodeUID));
 									if (!TextureNode)
 									{
diff --git a/Engine/Plugins/Runtime/Interchange/Source/Parsers/Fbx/Private/FbxScene.cpp b/Engine/Plugins/Runtime/Interchange/Source/Parsers/Fbx/Private/FbxScene.cpp
--- a/Engine/Plugins/Runtime/Interchange/Source/Parsers/Fbx/Private/FbxScene.cpp
+++ b/Engine/Plugins/Runtime/Interchange/Source/Parsers/Fbx/Private/FbxScene.cpp
@@ -71,8 +71,8 @@ namespace UE
 
 			UInterchangeSceneNode* FFbxScene::CreateTransformNode(UInterchangeBaseNodeContainer& NodeContainer, const FString& NodeName, const FString& NodeUniqueID, TArray<FString>& JSonErrorMessages)
 			{
-				FName DisplayLabel(*NodeName);
-				FName NodeUID(*NodeUniqueID);
+				FString DisplayLabel(NodeName);
+				FString NodeUID(NodeUniqueID);
 				UInterchangeSceneNode* TransformNode = NewObject<UInterchangeSceneNode>(&NodeContainer, NAME_None);
 				if (!ensure(TransformNode))
 				{
diff --git a/Engine/Plugins/Runtime/Interchange/Source/Parsers/Fbx/Private/FbxSkeletalMesh.cpp b/Engine/Plugins/Runtime/Interchange/Source/Parsers/Fbx/Private/FbxSkeletalMesh.cpp
--- a/Engine/Plugins/Runtime/Interchange/Source/Parsers/Fbx/Private/FbxSkeletalMesh.cpp
+++ b/Engine/Plugins/Runtime/Interchange/Source/Parsers/Fbx/Private/FbxSkeletalMesh.cpp
@@ -186,14 +186,6 @@ namespace UE
 
 					UInterchangeSkeletalMeshNode* SkeletalMeshNode = nullptr;
 
-					bool bOverrideSkeletalMeshName = false;
-					FString OverrideSkeletalMeshName;
-					if (SkeletalMeshCount == 1)
-					{
-						bOverrideSkeletalMeshName = true;
-						OverrideSkeletalMeshName = FPaths::GetBaseFilename(SourceFilename);
-					}
-
 					for (int32 LODIndex = 0; LODIndex < MaxNumberOfLOD; LODIndex++)
 					{
 						//Array use to retrieve the payload data, so we easily find the request payload data
@@ -242,7 +234,7 @@ namespace UE
 						{
 							continue;
 						}
-						FString SkeletonName = (bOverrideSkeletalMeshName ? OverrideSkeletalMeshName : FFbxHelper::GetFbxObjectName(SkelMeshNodeArray[0])) + TEXT("_Skeleton");
+						FString SkeletonName = FFbxHelper::GetFbxObjectName(SortedJoints[0]) + TEXT("_Skeleton");
 						if (LODIndex > 0)
 						{
 							SkeletonName += TEXT("_") + FString::FromInt(LODIndex);
@@ -261,7 +253,7 @@ namespace UE
 						FbxNode* RootNode = SkelMeshNodeArray[0];
 						if (LODIndex == 0)
 						{
-							FString SkeletalMeshName = (bOverrideSkeletalMeshName ? OverrideSkeletalMeshName : FFbxHelper::GetFbxObjectName(RootNode));
+							FString SkeletalMeshName = FFbxHelper::GetFbxObjectName(RootNode);
 							FString SkeletalMeshUniqueID = TEXT("\\SkeletalMesh\\") + FFbxHelper::GetFbxNodeHierarchyName(RootNode);
 							SkeletalMeshNode = CreateSkeletalMeshNode(NodeContainer, SkeletalMeshName, SkeletalMeshUniqueID, JSonErrorMessages);
 						}
@@ -296,7 +288,7 @@ namespace UE
 								}
 								//Find an existing material node
 								FString MaterialName = FFbxHelper::GetFbxObjectName(SurfaceMaterial);
-								FName NodeUID(*MaterialName);
+								FString NodeUID(MaterialName);
 								UInterchangeBaseNode* MaterialNode = NodeContainer.GetNode(NodeUID);
 								if (MaterialNode != nullptr && Cast<UInterchangeMaterialNode>(MaterialNode))
 								{
@@ -398,8 +390,8 @@ namespace UE
 
 			UInterchangeSkeletalMeshNode* FFbxSkeletalMesh::CreateSkeletalMeshNode(UInterchangeBaseNodeContainer& NodeContainer, const FString& NodeName, const FString& NodeUniqueID, TArray<FString>& JSonErrorMessages)
 			{
-				FName DisplayLabel(*NodeName);
-				FName NodeUID(*NodeUniqueID);
+				FString DisplayLabel(NodeName);
+				FString NodeUID(NodeUniqueID);
 				UInterchangeSkeletalMeshNode* SkeletalMeshNode = NewObject<UInterchangeSkeletalMeshNode>(&NodeContainer, NAME_None);
 				if (!ensure(SkeletalMeshNode))
 				{
@@ -414,8 +406,8 @@ namespace UE
 
 			UInterchangeSkeletalMeshLodDataNode* FFbxSkeletalMesh::CreateSkeletalMeshLodDataNode(UInterchangeBaseNodeContainer& NodeContainer, const FString& NodeName, const FString& NodeUniqueID, TArray<FString>& JSonErrorMessages)
 			{
-				FName DisplayLabel(*NodeName);
-				FName NodeUID(*NodeUniqueID);
+				FString DisplayLabel(NodeName);
+				FString NodeUID(NodeUniqueID);
 				UInterchangeSkeletalMeshLodDataNode* SkeletalMeshLodDataNode = NewObject<UInterchangeSkeletalMeshLodDataNode>(&NodeContainer, NAME_None);
 				if (!ensure(SkeletalMeshLodDataNode))
 				{
diff --git a/Engine/Plugins/Runtime/Interchange/Source/Parsers/Fbx/Private/FbxSkeleton.cpp b/Engine/Plugins/Runtime/Interchange/Source/Parsers/Fbx/Private/FbxSkeleton.cpp
--- a/Engine/Plugins/Runtime/Interchange/Source/Parsers/Fbx/Private/FbxSkeleton.cpp
+++ b/Engine/Plugins/Runtime/Interchange/Source/Parsers/Fbx/Private/FbxSkeleton.cpp
@@ -43,14 +43,6 @@ namespace UE
 						}
 					}
 
-					bool bOverrideSkeletalMeshName = false;
-					FString OverrideSkeletalMeshName;
-					if (SkeletalMeshCount == 1)
-					{
-						bOverrideSkeletalMeshName = true;
-						OverrideSkeletalMeshName = FPaths::GetBaseFilename(SourceFilename);
-					}
-
 					for (int32 LODIndex = 0; LODIndex < MaxNumberOfLOD; LODIndex++)
 					{
 						TArray<FbxNode*> SkelMeshNodeArray;
@@ -96,7 +88,7 @@ namespace UE
 						{
 							continue;
 						}
-						FString SkeletonName = (bOverrideSkeletalMeshName ? OverrideSkeletalMeshName : FFbxHelper::GetFbxObjectName(SkelMeshNodeArray[0])) + TEXT("_Skeleton");
+						FString SkeletonName = FFbxHelper::GetFbxObjectName(SortedJoints[0]) + TEXT("_Skeleton");
 						if(LODIndex > 0)
 						{
 							SkeletonName += TEXT("_") + FString::FromInt(LODIndex);
@@ -153,11 +145,11 @@ namespace UE
 								JointNode->SetParentUID(*ParentJointUniqueID);
 							}
 						}
-						FName SkeletonRootJointIDValidator;
+						FString SkeletonRootJointIDValidator;
 						if (SkeletonNode->GetCustomRootJointID(SkeletonRootJointIDValidator))
 						{
 							UInterchangeBaseNode* RootNode = NodeContainer.GetNode(SkeletonRootJointIDValidator);
-							FName RootNodeDisplayLabel = RootNode->GetDisplayLabel();
+							FString RootNodeDisplayLabel = RootNode->GetDisplayLabel();
 						}
 					}
 				}
@@ -165,8 +157,8 @@ namespace UE
 
 			UInterchangeSkeletonNode* FFbxSkeleton::CreateSkeletonNode(UInterchangeBaseNodeContainer& NodeContainer, const FString& NodeName, const FString& UniqueID, TArray<FString>& JSonErrorMessages)
 			{
-				FName DisplayLabel = *NodeName;
-				FName NodeUID(*UniqueID);
+				FString DisplayLabel = NodeName;
+				FString NodeUID(UniqueID);
 
 				UInterchangeBaseNode* ExistingNode = NodeContainer.GetNode(NodeUID);
 				if (ExistingNode)
@@ -189,8 +181,8 @@ namespace UE
 
 			UInterchangeJointNode* FFbxSkeleton::CreateJointNode(UInterchangeBaseNodeContainer& NodeContainer, const FString& NodeName, const FString& UniqueID, TArray<FString>& JSonErrorMessages)
 			{
-				FName DisplayLabel = *NodeName;
-				FName NodeUID(*UniqueID);
+				FString DisplayLabel = NodeName;
+				FString NodeUID(UniqueID);
 				UInterchangeJointNode* JointNode = NewObject<UInterchangeJointNode>(&NodeContainer, NAME_None);
 				if (!ensure(JointNode))
 				{
diff --git a/Engine/Source/Runtime/Interchange/Core/Private/Nodes/InterchangeBaseNode.cpp b/Engine/Source/Runtime/Interchange/Core/Private/Nodes/InterchangeBaseNode.cpp
--- a/Engine/Source/Runtime/Interchange/Core/Private/Nodes/InterchangeBaseNode.cpp
+++ b/Engine/Source/Runtime/Interchange/Core/Private/Nodes/InterchangeBaseNode.cpp
@@ -10,7 +10,7 @@
 
 
 
-void UInterchangeBaseNode::InitializeNode(const FName& UniqueID, const FName& DisplayLabel)
+void UInterchangeBaseNode::InitializeNode(const FString& UniqueID, const FString& DisplayLabel)
 {
 	UE::Interchange::EAttributeStorageResult Result = Attributes->RegisterAttribute(UE::Interchange::FBaseNodeStaticData::UniqueIDKey(), UniqueID, UE::Interchange::EAttributeProperty::NoHash);
 	if (!IsAttributeStorageResultSuccess(Result))
@@ -41,43 +41,43 @@ bool UInterchangeBaseNode::HasAttribute(const UE::Interchange::FAttributeKey& No
 }
 
 
-FName UInterchangeBaseNode::GetUniqueID() const
+FString UInterchangeBaseNode::GetUniqueID() const
 {
 	ensure(bIsInitialized);
-	FName UniqueID = NAME_None;
-	Attributes->GetAttributeHandle<FName>(UE::Interchange::FBaseNodeStaticData::UniqueIDKey()).Get(UniqueID);
+	FString UniqueID;
+	Attributes->GetAttributeHandle<FString>(UE::Interchange::FBaseNodeStaticData::UniqueIDKey()).Get(UniqueID);
 	return UniqueID;
 }
 
-FName UInterchangeBaseNode::GetDisplayLabel() const
+FString UInterchangeBaseNode::GetDisplayLabel() const
 {
 	ensure(bIsInitialized);
 	checkSlow(Attributes->ContainAttribute(UE::Interchange::FBaseNodeStaticData::DisplayLabelKey()));
-	FName DisplayLabel = NAME_None;
-	Attributes->GetAttributeHandle<FName>(UE::Interchange::FBaseNodeStaticData::DisplayLabelKey()).Get(DisplayLabel);
+	FString DisplayLabel;
+	Attributes->GetAttributeHandle<FString>(UE::Interchange::FBaseNodeStaticData::DisplayLabelKey()).Get(DisplayLabel);
 	return DisplayLabel;
 }
 
-bool UInterchangeBaseNode::SetDisplayLabel(FName DisplayLabel)
+bool UInterchangeBaseNode::SetDisplayLabel(const FString& DisplayLabel)
 {
 	UE::Interchange::EAttributeStorageResult Result = Attributes->RegisterAttribute(UE::Interchange::FBaseNodeStaticData::DisplayLabelKey(), DisplayLabel);
 	if (IsAttributeStorageResultSuccess(Result))
 	{
-		UE::Interchange::FAttributeStorage::TAttributeHandle<FName> Handle = Attributes->GetAttributeHandle<FName>(UE::Interchange::FBaseNodeStaticData::DisplayLabelKey());
+		UE::Interchange::FAttributeStorage::TAttributeHandle<FString> Handle = Attributes->GetAttributeHandle<FString>(UE::Interchange::FBaseNodeStaticData::DisplayLabelKey());
 		return Handle.IsValid();
 	}
 	return false;
 }
 
-FName UInterchangeBaseNode::GetParentUID() const
+FString UInterchangeBaseNode::GetParentUID() const
 {
 	if (!Attributes->ContainAttribute(UE::Interchange::FBaseNodeStaticData::ParentIDKey()))
 	{
 		return InvalidNodeUID();
 	}
 
-	FName ParentUniqueID = NAME_None;
-	UE::Interchange::FAttributeStorage::TAttributeHandle<FName> Handle = Attributes->GetAttributeHandle<FName>(UE::Interchange::FBaseNodeStaticData::ParentIDKey());
+	FString ParentUniqueID;
+	UE::Interchange::FAttributeStorage::TAttributeHandle<FString> Handle = Attributes->GetAttributeHandle<FString>(UE::Interchange::FBaseNodeStaticData::ParentIDKey());
 	if(Handle.IsValid())
 	{
 		Handle.Get(ParentUniqueID);
@@ -86,12 +86,12 @@ FName UInterchangeBaseNode::GetParentUID() const
 	return InvalidNodeUID();
 }
 
-bool UInterchangeBaseNode::SetParentUID(FName ParentUID)
+bool UInterchangeBaseNode::SetParentUID(const FString& ParentUID)
 {
 	UE::Interchange::EAttributeStorageResult Result = Attributes->RegisterAttribute(UE::Interchange::FBaseNodeStaticData::ParentIDKey(), ParentUID);
 	if(IsAttributeStorageResultSuccess(Result))
 	{
-		UE::Interchange::FAttributeStorage::TAttributeHandle<FName> Handle = Attributes->GetAttributeHandle<FName>(UE::Interchange::FBaseNodeStaticData::ParentIDKey());
+		UE::Interchange::FAttributeStorage::TAttributeHandle<FString> Handle = Attributes->GetAttributeHandle<FString>(UE::Interchange::FBaseNodeStaticData::ParentIDKey());
 		return Handle.IsValid();
 	}
 	return false;
@@ -102,17 +102,17 @@ int32 UInterchangeBaseNode::GetDependeciesCount() const
 	return Dependencies.GetCount();
 }
 
-void UInterchangeBaseNode::GetDependecies(TArray<FName>& OutDependencies ) const
+void UInterchangeBaseNode::GetDependecies(TArray<FString>& OutDependencies ) const
 {
 	Dependencies.GetNames(OutDependencies);
 }
 
-bool UInterchangeBaseNode::SetDependencyUID(FName DependencyUID)
+bool UInterchangeBaseNode::SetDependencyUID(const FString& DependencyUID)
 {
 	return Dependencies.AddName(DependencyUID);
 }
 
-bool UInterchangeBaseNode::RemoveDependencyUID(FName DependencyUID)
+bool UInterchangeBaseNode::RemoveDependencyUID(const FString& DependencyUID)
 {
 	return Dependencies.RemoveName(DependencyUID);
 }
@@ -150,9 +150,9 @@ class UClass* UInterchangeBaseNode::GetAssetClass() const
 	return nullptr;
 }
 
-FName UInterchangeBaseNode::InvalidNodeUID()
+FString UInterchangeBaseNode::InvalidNodeUID()
 {
-	return NAME_None;
+	return FString();
 }
 
 void UInterchangeBaseNode::ApplyAllCustomAttributeToAsset(UObject* Object) const
@@ -198,8 +198,8 @@ void UInterchangeBaseNode::Serialize(FArchive& Ar)
 	if (Ar.IsLoading())
 	{
 		//The node is consider Initialize if the UniqueID and the Display label are set properly
-		bIsInitialized = (Attributes->GetAttributeHandle<FName>(UE::Interchange::FBaseNodeStaticData::UniqueIDKey()).IsValid() &&
-						  Attributes->GetAttributeHandle<FName>(UE::Interchange::FBaseNodeStaticData::DisplayLabelKey()).IsValid());
+		bIsInitialized = (Attributes->GetAttributeHandle<FString>(UE::Interchange::FBaseNodeStaticData::UniqueIDKey()).IsValid() &&
+						  Attributes->GetAttributeHandle<FString>(UE::Interchange::FBaseNodeStaticData::DisplayLabelKey()).IsValid());
 
 	}
 }
diff --git a/Engine/Source/Runtime/Interchange/Core/Private/Nodes/InterchangeBaseNodeContainer.cpp b/Engine/Source/Runtime/Interchange/Core/Private/Nodes/InterchangeBaseNodeContainer.cpp
--- a/Engine/Source/Runtime/Interchange/Core/Private/Nodes/InterchangeBaseNodeContainer.cpp
+++ b/Engine/Source/Runtime/Interchange/Core/Private/Nodes/InterchangeBaseNodeContainer.cpp
@@ -12,13 +12,13 @@
 #include "UObject/ObjectMacros.h"
 #include "UObject/UObjectIterator.h"
 
-FName UInterchangeBaseNodeContainer::AddNode(UInterchangeBaseNode* Node)
+FString UInterchangeBaseNodeContainer::AddNode(UInterchangeBaseNode* Node)
 {
 	if (!Node)
 	{
 		return UInterchangeBaseNode::InvalidNodeUID();
 	}
-	FName NodeUniqueID = Node->GetUniqueID();
+	FString NodeUniqueID = Node->GetUniqueID();
 	if (NodeUniqueID == UInterchangeBaseNode::InvalidNodeUID())
 	{
 		return UInterchangeBaseNode::InvalidNodeUID();
@@ -30,12 +30,12 @@ FName UInterchangeBaseNodeContainer::AddNode(UInterchangeBaseNode* Node)
 		return NodeUniqueID;
 	}
 
-	if (Node->GetDisplayLabel() == NAME_None)
+	if (Node->GetDisplayLabel().IsEmpty())
 	{
 		//Replace None by Null, since None name will be interpret like NAME_None which will not work with UObject creation
 		//UObject Creation will name it ClassName_X instead of None
 		//TODO Log an warning to the user
-		Node->SetDisplayLabel(FName("Null"));
+		Node->SetDisplayLabel(FString(TEXT("Null")));
 	}
 
 	//Copy the node
@@ -43,7 +43,7 @@ FName UInterchangeBaseNodeContainer::AddNode(UInterchangeBaseNode* Node)
 	return NodeUniqueID;
 }
 
-bool UInterchangeBaseNodeContainer::IsNodeUIDValid(FName NodeUniqueID) const
+bool UInterchangeBaseNodeContainer::IsNodeUIDValid(const FString& NodeUniqueID) const
 {
 	if (NodeUniqueID == UInterchangeBaseNode::InvalidNodeUID())
 	{
@@ -52,17 +52,17 @@ bool UInterchangeBaseNodeContainer::IsNodeUIDValid(FName NodeUniqueID) const
 	return Nodes.Contains(NodeUniqueID);
 }
 
-void UInterchangeBaseNodeContainer::IterateNodes(TFunctionRef<void(FName, UInterchangeBaseNode*)> IterationLambda)
+void UInterchangeBaseNodeContainer::IterateNodes(TFunctionRef<void(const FString&, UInterchangeBaseNode*)> IterationLambda)
 {
-	for (TPair<FName, UInterchangeBaseNode*>& NodeKeyValue : Nodes)
+	for (TPair<FString, UInterchangeBaseNode*>& NodeKeyValue : Nodes)
 	{
 		IterationLambda(NodeKeyValue.Key, NodeKeyValue.Value);
 	}
 }
 
-void UInterchangeBaseNodeContainer::GetRoots(TArray<FName>& RootNodes)
+void UInterchangeBaseNodeContainer::GetRoots(TArray<FString>& RootNodes)
 {
-	for (TPair<FName, UInterchangeBaseNode*>& NodeKeyValue : Nodes)
+	for (TPair<FString, UInterchangeBaseNode*>& NodeKeyValue : Nodes)
 	{
 		if (NodeKeyValue.Value->GetParentUID() == UInterchangeBaseNode::InvalidNodeUID())
 		{
@@ -71,7 +71,7 @@ void UInterchangeBaseNodeContainer::GetRoots(TArray<FName>& RootNodes)
 	}
 }
 
-UInterchangeBaseNode* UInterchangeBaseNodeContainer::GetNode(FName NodeUniqueID)
+UInterchangeBaseNode* UInterchangeBaseNodeContainer::GetNode(const FString& NodeUniqueID)
 {
 	if (NodeUniqueID == UInterchangeBaseNode::InvalidNodeUID())
 	{
@@ -85,12 +85,12 @@ UInterchangeBaseNode* UInterchangeBaseNodeContainer::GetNode(FName NodeUniqueID)
 	return Node;
 }
 
-const UInterchangeBaseNode* UInterchangeBaseNodeContainer::GetNode(FName NodeUniqueID) const
+const UInterchangeBaseNode* UInterchangeBaseNodeContainer::GetNode(const FString& NodeUniqueID) const
 {
 	return const_cast<UInterchangeBaseNodeContainer*>(this)->GetNode(NodeUniqueID);
 }
 
-bool UInterchangeBaseNodeContainer::SetNodeParentUID(FName NodeUniqueID, FName NewParentNodeUID)
+bool UInterchangeBaseNodeContainer::SetNodeParentUID(const FString& NodeUniqueID, const FString& NewParentNodeUID)
 {
 	if (!Nodes.Contains(NodeUniqueID))
 	{
@@ -105,16 +105,16 @@ bool UInterchangeBaseNodeContainer::SetNodeParentUID(FName NodeUniqueID, FName N
 	return true;
 }
 
-int32 UInterchangeBaseNodeContainer::GetNodeChildrenCount(FName NodeUniqueID) const
+int32 UInterchangeBaseNodeContainer::GetNodeChildrenCount(const FString& NodeUniqueID) const
 {
-	TArray<FName> ChildrenUIDs = GetNodeChildrenUIDs(NodeUniqueID);
+	TArray<FString> ChildrenUIDs = GetNodeChildrenUIDs(NodeUniqueID);
 	return ChildrenUIDs.Num();
 }
 
-TArray<FName> UInterchangeBaseNodeContainer::GetNodeChildrenUIDs(FName NodeUniqueID) const
+TArray<FString> UInterchangeBaseNodeContainer::GetNodeChildrenUIDs(const FString& NodeUniqueID) const
 {
-	TArray<FName> OutChildrenUIDs;
-	for (const TPair<FName, UInterchangeBaseNode*>& NodeKeyValue : Nodes)
+	TArray<FString> OutChildrenUIDs;
+	for (const TPair<FString, UInterchangeBaseNode*>& NodeKeyValue : Nodes)
 	{
 		if (NodeKeyValue.Value->GetParentUID() == NodeUniqueID)
 		{
@@ -124,12 +124,12 @@ TArray<FName> UInterchangeBaseNodeContainer::GetNodeChildrenUIDs(FName NodeUniqu
 	return OutChildrenUIDs;
 }
 
-UInterchangeBaseNode* UInterchangeBaseNodeContainer::GetNodeChildren(FName NodeUniqueID, int32 ChildIndex)
+UInterchangeBaseNode* UInterchangeBaseNodeContainer::GetNodeChildren(const FString& NodeUniqueID, int32 ChildIndex)
 {
 	return GetNodeChildrenInternal(NodeUniqueID, ChildIndex);
 }
 
-const UInterchangeBaseNode* UInterchangeBaseNodeContainer::GetNodeChildren(FName NodeUniqueID, int32 ChildIndex) const
+const UInterchangeBaseNode* UInterchangeBaseNodeContainer::GetNodeChildren(const FString& NodeUniqueID, int32 ChildIndex) const
 {
 	return const_cast<UInterchangeBaseNodeContainer*>(this)->GetNodeChildrenInternal(NodeUniqueID, ChildIndex);
 }
@@ -153,7 +153,7 @@ void UInterchangeBaseNodeContainer::SerializeNodeContainerData(FArchive& Ar)
 			BaseNode->Serialize(Ar);
 		};
 
-		for (TPair<FName, UInterchangeBaseNode*> NodePair : Nodes)
+		for (TPair<FString, UInterchangeBaseNode*> NodePair : Nodes)
 		{
 			SerializeNodePair(NodePair.Value);
 		}
@@ -218,9 +218,9 @@ void UInterchangeBaseNodeContainer::LoadFromFile(const FString& Filename)
 	SerializeNodeContainerData(Ar);
 }
 
-UInterchangeBaseNode* UInterchangeBaseNodeContainer::GetNodeChildrenInternal(FName NodeUniqueID, int32 ChildIndex)
+UInterchangeBaseNode* UInterchangeBaseNodeContainer::GetNodeChildrenInternal(const FString& NodeUniqueID, int32 ChildIndex)
 {
-	TArray<FName> ChildrenUIDs = GetNodeChildrenUIDs(NodeUniqueID);
+	TArray<FString> ChildrenUIDs = GetNodeChildrenUIDs(NodeUniqueID);
 	if (!ChildrenUIDs.IsValidIndex(ChildIndex))
 	{
 		return nullptr;
diff --git a/Engine/Source/Runtime/Interchange/Core/Private/Tests/StorageTest.cpp b/Engine/Source/Runtime/Interchange/Core/Private/Tests/StorageTest.cpp
--- a/Engine/Source/Runtime/Interchange/Core/Private/Tests/StorageTest.cpp
+++ b/Engine/Source/Runtime/Interchange/Core/Private/Tests/StorageTest.cpp
@@ -16,14 +16,13 @@ bool FAttributeStorageTest::RunTest(const FString& Parameters)
 {
 	using namespace UE::Interchange;
 
-	FString KeyPrefixString = TEXT("TestKey");
-	FName KeyPrefix(*KeyPrefixString);
+	FString KeyPrefix = TEXT("TestKey");
 	int32 KeySuffixCounter = 1;
 	auto CreateUniqueKey = [&KeyPrefix, &KeySuffixCounter]()->FAttributeKey
 	{
 		check(IsInGameThread());
-		KeyPrefix.SetNumber(KeySuffixCounter++);
-		const FAttributeKey UniqueKeyRet = FAttributeKey(KeyPrefix);
+		FString KeyFinal = KeyPrefix + TEXT("_") + FString::FromInt(KeySuffixCounter++);
+		const FAttributeKey UniqueKeyRet = FAttributeKey(KeyFinal);
 		return UniqueKeyRet;
 	};
 
diff --git a/Engine/Source/Runtime/Interchange/Core/Private/Types/AttributeStorage.cpp b/Engine/Source/Runtime/Interchange/Core/Private/Types/AttributeStorage.cpp
--- a/Engine/Source/Runtime/Interchange/Core/Private/Types/AttributeStorage.cpp
+++ b/Engine/Source/Runtime/Interchange/Core/Private/Types/AttributeStorage.cpp
@@ -17,49 +17,49 @@ namespace UE
 			//Errors
 			if (HasAttributeStorageResult(Result, EAttributeStorageResult::Operation_Error_AttributeAllocationCorrupted))
 			{
-				UE_LOG(LogInterchangeCore, Error, TEXT("Attribute storage operation [%s] Key[%s]: Storage is corrupted."), *OperationName, *(AttributeKey.Key.ToString()));
+				UE_LOG(LogInterchangeCore, Error, TEXT("Attribute storage operation [%s] Key[%s]: Storage is corrupted."), *OperationName, *(AttributeKey.Key));
 			}
 
 			//////////////////////////////////////////////////////////////////////////
 			//Warning
 			if (HasAttributeStorageResult(Result, EAttributeStorageResult::Operation_Error_CannotFoundKey))
 			{
-				UE_LOG(LogInterchangeCore, Warning, TEXT("Attribute storage operation [%s] Key[%s]: Cannot find attribute key."), *OperationName, *(AttributeKey.Key.ToString()));
+				UE_LOG(LogInterchangeCore, Warning, TEXT("Attribute storage operation [%s] Key[%s]: Cannot find attribute key."), *OperationName, *(AttributeKey.Key));
 			}
 
 			if (HasAttributeStorageResult(Result, EAttributeStorageResult::Operation_Error_CannotOverrideAttribute))
 			{
-				UE_LOG(LogInterchangeCore, Warning, TEXT("Attribute storage operation [%s] Key[%s]: Cannot override attribute."), *OperationName, *(AttributeKey.Key.ToString()));
+				UE_LOG(LogInterchangeCore, Warning, TEXT("Attribute storage operation [%s] Key[%s]: Cannot override attribute."), *OperationName, *(AttributeKey.Key));
 			}
 
 			if (HasAttributeStorageResult(Result, EAttributeStorageResult::Operation_Error_CannotRemoveAttribute))
 			{
-				UE_LOG(LogInterchangeCore, Warning, TEXT("Attribute storage operation [%s] Key[%s]: Cannot remove an attribute."), *OperationName, *(AttributeKey.Key.ToString()));
+				UE_LOG(LogInterchangeCore, Warning, TEXT("Attribute storage operation [%s] Key[%s]: Cannot remove an attribute."), *OperationName, *(AttributeKey.Key));
 			}
 
 			if (HasAttributeStorageResult(Result, EAttributeStorageResult::Operation_Error_WrongSize))
 			{
-				UE_LOG(LogInterchangeCore, Warning, TEXT("Attribute storage operation [%s] Key[%s]: Stored attribute value size do not match parameter value size."), *OperationName, *(AttributeKey.Key.ToString()));
+				UE_LOG(LogInterchangeCore, Warning, TEXT("Attribute storage operation [%s] Key[%s]: Stored attribute value size do not match parameter value size."), *OperationName, *(AttributeKey.Key));
 			}
 
 			if (HasAttributeStorageResult(Result, EAttributeStorageResult::Operation_Error_WrongType))
 			{
-				UE_LOG(LogInterchangeCore, Warning, TEXT("Attribute storage operation [%s] Key[%s]: Stored attribute value type do not match parameter value type."), *OperationName, *(AttributeKey.Key.ToString()));
+				UE_LOG(LogInterchangeCore, Warning, TEXT("Attribute storage operation [%s] Key[%s]: Stored attribute value type do not match parameter value type."), *OperationName, *(AttributeKey.Key));
 			}
 
 			if (HasAttributeStorageResult(Result, EAttributeStorageResult::Operation_Error_UnsupportedType))
 			{
-				UE_LOG(LogInterchangeCore, Warning, TEXT("Attribute storage operation [%s] Key[%s]: The storage do not support this type of value, please add unsupported type to EAttributeTypes."), *OperationName, *(AttributeKey.Key.ToString()));
+				UE_LOG(LogInterchangeCore, Warning, TEXT("Attribute storage operation [%s] Key[%s]: The storage do not support this type of value, please add unsupported type to EAttributeTypes."), *OperationName, *(AttributeKey.Key));
 			}
 
 			if (HasAttributeStorageResult(Result, EAttributeStorageResult::Operation_Error_InvalidStorage))
 			{
-				UE_LOG(LogInterchangeCore, Warning, TEXT("Attribute storage operation [%s] Key[%s]: The storage is invalid (NULL)."), *OperationName, *(AttributeKey.Key.ToString()));
+				UE_LOG(LogInterchangeCore, Warning, TEXT("Attribute storage operation [%s] Key[%s]: The storage is invalid (NULL)."), *OperationName, *(AttributeKey.Key));
 			}
 
 			if (HasAttributeStorageResult(Result, EAttributeStorageResult::Operation_Error_InvalidMultiSizeValueData))
 			{
-				UE_LOG(LogInterchangeCore, Warning, TEXT("Attribute storage operation [%s] Key[%s]: Cannot retrieve the multisize value data pointer."), *OperationName, *(AttributeKey.Key.ToString()));
+				UE_LOG(LogInterchangeCore, Warning, TEXT("Attribute storage operation [%s] Key[%s]: Cannot retrieve the multisize value data pointer."), *OperationName, *(AttributeKey.Key));
 			}
 		}
 
diff --git a/Engine/Source/Runtime/Interchange/Core/Public/Nodes/InterchangeBaseNode.h b/Engine/Source/Runtime/Interchange/Core/Public/Nodes/InterchangeBaseNode.h
--- a/Engine/Source/Runtime/Interchange/Core/Public/Nodes/InterchangeBaseNode.h
+++ b/Engine/Source/Runtime/Interchange/Core/Public/Nodes/InterchangeBaseNode.h
@@ -188,7 +188,7 @@ namespace UE
 				Attributes = InAttributes;
 				check(Attributes.IsValid());
 				FString BaseTryName = TEXT("__") + BaseKeyName;
-				KeyCount = FName(*BaseTryName);
+				KeyCount = BaseTryName;
 			}
 
 			int32 GetCount() const
@@ -211,7 +211,7 @@ namespace UE
 				return NameCount;
 			}
 
-			void GetNames(TArray<FName>& OutNames) const
+			void GetNames(TArray<FString>& OutNames) const
 			{
 				//The Class must be initialise properly before we can use it
 				TSharedPtr<FAttributeStorage, ESPMode::ThreadSafe> AttributePtr = Attributes.Pin();
@@ -239,17 +239,17 @@ namespace UE
 				for (int32 NameIndex = 0; NameIndex < NameCount; ++NameIndex)
 				{
 					FAttributeKey DepIndexKey = GetIndexKey(NameIndex);
-					FAttributeStorage::TAttributeHandle<FName> HandleName = AttributePtr->GetAttributeHandle<FName>(DepIndexKey);
+					FAttributeStorage::TAttributeHandle<FString> HandleName = AttributePtr->GetAttributeHandle<FString>(DepIndexKey);
 					if (!HandleName.IsValid())
 					{
 						continue;
 					}
-					FName& OutName = OutNames.AddDefaulted_GetRef();
+					FString& OutName = OutNames.AddDefaulted_GetRef();
 					HandleName.Get(OutName);
 				}
 			}
 
-			bool AddName(FName Name)
+			bool AddName(const FString& Name)
 			{
 				//The Class must be initialise properly before we can use it
 				TSharedPtr<FAttributeStorage, ESPMode::ThreadSafe> AttributePtr = Attributes.Pin();
@@ -280,7 +280,7 @@ namespace UE
 				NameIndex++;
 				Handle.Set(NameIndex);
 
-				EAttributeStorageResult AddNameResult = AttributePtr->RegisterAttribute<FName>(NameIndexKey, Name);
+				EAttributeStorageResult AddNameResult = AttributePtr->RegisterAttribute<FString>(NameIndexKey, Name);
 				if (!IsAttributeStorageResultSuccess(AddNameResult))
 				{
 					LogAttributeStorageErrors(AddNameResult, TEXT("FNameAttributeArrayHelper.AddName"), NameIndexKey);
@@ -289,7 +289,7 @@ namespace UE
 				return true;
 			}
 
-			bool RemoveName(FName NameToDelete)
+			bool RemoveName(const FString& NameToDelete)
 			{
 				//The Class must be initialise properly before we can use it
 				TSharedPtr<FAttributeStorage, ESPMode::ThreadSafe> AttributePtr = Attributes.Pin();
@@ -313,12 +313,12 @@ namespace UE
 				for (int32 NameIndex = 0; NameIndex < NameCount; ++NameIndex)
 				{
 					FAttributeKey DepIndexKey = GetIndexKey(NameIndex);
-					FAttributeStorage::TAttributeHandle<FName> HandleName = AttributePtr->GetAttributeHandle<FName>(DepIndexKey);
+					FAttributeStorage::TAttributeHandle<FString> HandleName = AttributePtr->GetAttributeHandle<FString>(DepIndexKey);
 					if (!HandleName.IsValid())
 					{
 						continue;
 					}
-					FName Name;
+					FString Name;
 					HandleName.Get(Name);
 					if (Name == NameToDelete)
 					{
@@ -334,7 +334,7 @@ namespace UE
 						EAttributeStorageResult UnregisterResult = AttributePtr->UnregisterAttribute(DepIndexKey);
 						if (IsAttributeStorageResultSuccess(UnregisterResult))
 						{
-							EAttributeStorageResult RegisterResult = AttributePtr->RegisterAttribute<FName>(NewDepIndexKey, Name);
+							EAttributeStorageResult RegisterResult = AttributePtr->RegisterAttribute<FString>(NewDepIndexKey, Name);
 							if (!IsAttributeStorageResultSuccess(RegisterResult))
 							{
 								LogAttributeStorageErrors(RegisterResult, TEXT("FNameAttributeArrayHelper.RemoveName"), NewDepIndexKey);
@@ -388,7 +388,7 @@ namespace UE
 			TWeakPtr<FAttributeStorage, ESPMode::ThreadSafe> Attributes = nullptr;
 
 			FAttributeKey KeyCount; //Assign in Initialize function, it will ensure if its the default value (NAME_None) when using the class
-			FAttributeKey GetKeyCount() const { ensure(KeyCount.Key != NAME_None); return KeyCount; }
+			FAttributeKey GetKeyCount() const { ensure(!KeyCount.Key.IsEmpty()); return KeyCount; }
 
 			FAttributeKey GetIndexKey(int32 Index) const
 			{
@@ -464,7 +464,7 @@ public:
 	 *
 	 */
 	UFUNCTION(BlueprintCallable, Category = "Interchange | Node")
-	void InitializeNode(const FName& UniqueID, const FName& DisplayLabel);
+	void InitializeNode(const FString& UniqueID, const FString& DisplayLabel);
 
 	/**
 	 * Return the node type name of the class, we use this when reporting error
@@ -524,33 +524,33 @@ public:
 	 *
 	 */
 	UFUNCTION(BlueprintCallable, Category = "Interchange | Node")
-	FName GetUniqueID() const;
+	FString GetUniqueID() const;
 
 	/**
 	 * Return the display label.
 	 *
 	 */
 	UFUNCTION(BlueprintCallable, Category = "Interchange | Node")
-	FName GetDisplayLabel() const;
+	FString GetDisplayLabel() const;
 
 	/**
 	 * Change the display label.
 	 *
 	 */
 	UFUNCTION(BlueprintCallable, Category = "Interchange | Node")
-	bool SetDisplayLabel(FName DisplayName);
+	bool SetDisplayLabel(const FString& DisplayName);
 
 	/**
 	 * Return the parent unique id. In case the attribute does not exist it will return InvalidNodeUID()
 	 */
 	UFUNCTION(BlueprintCallable, Category = "Interchange | Node")
-	FName GetParentUID() const;
+	FString GetParentUID() const;
 
 	/**
 	 * Set the parent unique id.
 	 */
 	UFUNCTION(BlueprintCallable, Category = "Interchange | Node")
-	bool SetParentUID(FName ParentUID);
+	bool SetParentUID(const FString& ParentUID);
 
 	/**
 	 * This function allow to retrieve the number of dependencies for this object.
@@ -564,19 +564,19 @@ public:
 	 * 
 	 */
 	UFUNCTION(BlueprintCallable, Category = "Interchange | Node")
-	void GetDependecies(TArray<FName>& OutDependencies ) const;
+	void GetDependecies(TArray<FString>& OutDependencies ) const;
 
 	/**
 	 * Add one dependency to this object.
 	 */
 	UFUNCTION(BlueprintCallable, Category = "Interchange | Node")
-	bool SetDependencyUID(FName DependencyUID);
+	bool SetDependencyUID(const FString& DependencyUID);
 
 	/**
 	 * Remove one dependency from this object.
 	 */
 	UFUNCTION(BlueprintCallable, Category = "Interchange | Node")
-	bool RemoveDependencyUID(FName DependencyUID);
+	bool RemoveDependencyUID(const FString& DependencyUID);
 
 	/**
 	 * IsEnable true mean that the node will be import/export, if false it will be discarded.
@@ -608,7 +608,7 @@ public:
 	virtual class UClass* GetAssetClass() const;
 
 	/** Return the invalid unique ID */
-	static FName InvalidNodeUID();
+	static FString InvalidNodeUID();
 
 	/**
 	 * Each Attribute that was set and have a delegate set for the specified UObject->UClass will
diff --git a/Engine/Source/Runtime/Interchange/Core/Public/Nodes/InterchangeBaseNodeContainer.h b/Engine/Source/Runtime/Interchange/Core/Public/Nodes/InterchangeBaseNodeContainer.h
--- a/Engine/Source/Runtime/Interchange/Core/Public/Nodes/InterchangeBaseNodeContainer.h
+++ b/Engine/Source/Runtime/Interchange/Core/Public/Nodes/InterchangeBaseNodeContainer.h
@@ -33,44 +33,44 @@ public:
 	 *
 	 */
 	UFUNCTION(BlueprintCallable, Category = "Interchange | Node Container")
-	FName AddNode(UInterchangeBaseNode* Node);
+	FString AddNode(UInterchangeBaseNode* Node);
 
 	/** Return true if the node unique ID exist in the container */
 	UFUNCTION(BlueprintCallable, Category = "Interchange | Node Container")
-	bool IsNodeUIDValid(FName NodeUniqueID) const;
+	bool IsNodeUIDValid(const FString& NodeUniqueID) const;
 
 	/** Unordered iteration of the all nodes */
-	void IterateNodes(TFunctionRef<void(FName, UInterchangeBaseNode*)> IterationLambda);
+	void IterateNodes(TFunctionRef<void(const FString&, UInterchangeBaseNode*)> IterationLambda);
 
 	/** Return all nodes that do not have any parent */
 	UFUNCTION(BlueprintCallable, Category = "Interchange | Node Container")
-	void GetRoots(TArray<FName>& RootNodes);
+	void GetRoots(TArray<FString>& RootNodes);
 
 	/** Get an node pointer */
 	UFUNCTION(BlueprintCallable, Category = "Interchange | Node Container")
-	UInterchangeBaseNode* GetNode(FName NodeUniqueID);
+	UInterchangeBaseNode* GetNode(const FString& NodeUniqueID);
 
 	/** Get an node pointer */
-	const UInterchangeBaseNode* GetNode(FName NodeUniqueID) const;
+	const UInterchangeBaseNode* GetNode(const FString& NodeUniqueID) const;
 
 	/** Set node ParentUID */
 	UFUNCTION(BlueprintCallable, Category = "Interchange | Node Container")
-	bool SetNodeParentUID(FName NodeUniqueID, FName NewParentNodeUID);
+	bool SetNodeParentUID(const FString& NodeUniqueID, const FString& NewParentNodeUID);
 
 	/** Get the node children count */
 	UFUNCTION(BlueprintCallable, Category = "Interchange | Node Container")
-	int32 GetNodeChildrenCount(FName NodeUniqueID) const;
+	int32 GetNodeChildrenCount(const FString& NodeUniqueID) const;
 
 	/** Get all children UID */
 	UFUNCTION(BlueprintCallable, Category = "Interchange | Node Container")
-	TArray<FName> GetNodeChildrenUIDs(FName NodeUniqueID) const;
+	TArray<FString> GetNodeChildrenUIDs(const FString& NodeUniqueID) const;
 
 	/** Get the node nth const children */
 	UFUNCTION(BlueprintCallable, Category = "Interchange | Node Container")
-	UInterchangeBaseNode* GetNodeChildren(FName NodeUniqueID, int32 ChildIndex);
+	UInterchangeBaseNode* GetNodeChildren(const FString& NodeUniqueID, int32 ChildIndex);
 
 	/** Get the node nth const children. Const version */
-	const UInterchangeBaseNode* GetNodeChildren(FName NodeUniqueID, int32 ChildIndex) const;
+	const UInterchangeBaseNode* GetNodeChildren(const FString& NodeUniqueID, int32 ChildIndex) const;
 
 	/**
 	 * This function serialize the node container and all node sub-objects point by it.
@@ -89,9 +89,9 @@ public:
 
 private:
 
-	UInterchangeBaseNode* GetNodeChildrenInternal(FName NodeUniqueID, int32 ChildIndex);
+	UInterchangeBaseNode* GetNodeChildrenInternal(const FString& NodeUniqueID, int32 ChildIndex);
 
 	/** Flat List of the nodes. Since the nodes are variable size, we store a pointer. */
 	UPROPERTY()
-	TMap<FName, UInterchangeBaseNode* > Nodes;
+	TMap<FString, UInterchangeBaseNode* > Nodes;
 };
diff --git a/Engine/Source/Runtime/Interchange/Core/Public/Types/AttributeStorage.h b/Engine/Source/Runtime/Interchange/Core/Public/Types/AttributeStorage.h
--- a/Engine/Source/Runtime/Interchange/Core/Public/Types/AttributeStorage.h
+++ b/Engine/Source/Runtime/Interchange/Core/Public/Types/AttributeStorage.h
@@ -20,7 +20,7 @@ namespace UE
 
 		struct FAttributeKey
 		{
-			FName Key = NAME_None;
+			FString Key;
 
 			FAttributeKey() = default;
 		
@@ -31,28 +31,28 @@ namespace UE
 
 			explicit FAttributeKey(const FName& Other)
 			{
-				Key = Other;
+				Key = Other.ToString();
 			}
 
 			explicit FAttributeKey(const FString& Other)
 			{
-				Key = Other.IsEmpty() ? NAME_None : FName(*Other);
+				Key = Other;
 			}
 
 			explicit FAttributeKey(const FText& Other)
 			{
-				Key = Other.IsEmpty() ? NAME_None : FName(*(Other.ToString()));
+				Key = Other.ToString();
 			}
 
 			explicit FAttributeKey(const TCHAR* Other)
 			{
 				const FString Helper(Other);
-				Key = Helper.IsEmpty() ? NAME_None : FName(Other);
+				Key = Helper;
 			}
 
 			FORCEINLINE const TCHAR* operator*() const
 			{
-				return *(Key.ToString());
+				return *(Key);
 			}
 
 			FORCEINLINE FAttributeKey& operator=(const FAttributeKey& Other)
@@ -63,26 +63,26 @@ namespace UE
 
 			FORCEINLINE FAttributeKey& operator=(const FName& Other)
 			{
-				Key = Other;
+				Key = Other.ToString();
 				return *this;
 			}
 
 			FORCEINLINE FAttributeKey& operator=(const FString& Other)
 			{
-				Key = Other.IsEmpty() ? NAME_None : FName(*Other);
+				Key = Other;
 				return *this;
 			}
 
 			FORCEINLINE FAttributeKey& operator=(const FText& Other)
 			{
-				Key = Other.IsEmpty() ? NAME_None : FName(*(Other.ToString()));
+				Key = Other.ToString();
 				return *this;
 			}
 
 			FORCEINLINE FAttributeKey& operator=(const TCHAR* Other)
 			{
 				const FString Helper(Other);
-				Key = Helper.IsEmpty() ? NAME_None : FName(Other);
+				Key = Helper;
 				return *this;
 			}
 		
@@ -98,7 +98,7 @@ namespace UE
 
 			FORCEINLINE bool operator<(const FAttributeKey& Other) const
 			{
-				return Key.LexicalLess(Other.Key);
+				return Key < Other.Key;
 			}
 
 			FORCEINLINE bool operator<=(const FAttributeKey& Other) const
@@ -108,7 +108,7 @@ namespace UE
 
 			FORCEINLINE bool operator>(const FAttributeKey& Other) const
 			{
-				return Other.Key.LexicalLess(Key);
+				return Key > Other.Key;
 			}
 
 			FORCEINLINE bool operator>=(const FAttributeKey& Other) const
@@ -118,18 +118,13 @@ namespace UE
 
 			friend FArchive& operator<<(FArchive& Ar, FAttributeKey& AttributeKey)
 			{
-				FString KeyString = AttributeKey.Key.ToString();
-				Ar << KeyString;
-				if(Ar.IsLoading())
-				{
-					AttributeKey.Key = FName(*KeyString);
-				}
+				Ar << AttributeKey.Key;
 				return Ar;
 			}
 
 			FORCEINLINE FString ToString() const
 			{
-				return Key.ToString();
+				return Key;
 			}
 		};
 
@@ -350,7 +345,7 @@ namespace UE
 			public:
 				TAttributeHandle()
 				: AttributeStorage(nullptr)
-				, Key(NAME_None)
+				, Key()
 				{}
 
 				/**
diff --git a/Engine/Source/Runtime/Interchange/Engine/Private/Tasks/InterchangeTaskCreateAsset.cpp b/Engine/Source/Runtime/Interchange/Engine/Private/Tasks/InterchangeTaskCreateAsset.cpp
--- a/Engine/Source/Runtime/Interchange/Engine/Private/Tasks/InterchangeTaskCreateAsset.cpp
+++ b/Engine/Source/Runtime/Interchange/Engine/Private/Tasks/InterchangeTaskCreateAsset.cpp
@@ -28,7 +28,7 @@ namespace UE
 			{
 				const UInterchangeSourceData* SourceData = AsyncHelper.SourceDatas[SourceIndex];
 				check(SourceData);
-				FString NodeDisplayName = Node->GetDisplayLabel().ToString();
+				FString NodeDisplayName = Node->GetDisplayLabel();
 				const FString BaseFileName = FPaths::GetBaseFilename(SourceData->GetFilename());
 
 				//Set the asset name and the package name
@@ -123,7 +123,7 @@ void UE::Interchange::FTaskCreatePackage::DoTask(ENamedThreads::Type CurrentThre
 			UE::Interchange::FImportAsyncHelper::FImportedAssetInfo& AssetInfo = ImportedInfos.AddDefaulted_GetRef();
 			AssetInfo.ImportAsset = NodeAsset;
 			AssetInfo.Factory = Factory;
-			AssetInfo.NodeUniqueId = Node->GetUniqueID().ToString();
+			AssetInfo.NodeUniqueId = Node->GetUniqueID();
 			Node->ReferenceObject = FSoftObjectPath(NodeAsset);
 		}
 	}
@@ -214,7 +214,7 @@ void UE::Interchange::FTaskCreateAsset::DoTask(ENamedThreads::Type CurrentThread
 			UE::Interchange::FImportAsyncHelper::FImportedAssetInfo& AssetInfo = ImportedInfos.AddDefaulted_GetRef();
 			AssetInfo.ImportAsset = NodeAsset;
 			AssetInfo.Factory = Factory;
-			AssetInfo.NodeUniqueId = Node->GetUniqueID().ToString();
+			AssetInfo.NodeUniqueId = Node->GetUniqueID();
 		}
 
 		Node->ReferenceObject = FSoftObjectPath(NodeAsset);
diff --git a/Engine/Source/Runtime/Interchange/Engine/Private/Tasks/InterchangeTaskParsing.cpp b/Engine/Source/Runtime/Interchange/Engine/Private/Tasks/InterchangeTaskParsing.cpp
--- a/Engine/Source/Runtime/Interchange/Engine/Private/Tasks/InterchangeTaskParsing.cpp
+++ b/Engine/Source/Runtime/Interchange/Engine/Private/Tasks/InterchangeTaskParsing.cpp
@@ -33,10 +33,10 @@ void UE::Interchange::FTaskParsing::DoTask(ENamedThreads::Type CurrentThread, co
 	//Parse each graph and prepare import task data, we will then be able to create all the task with the correct dependencies
 	struct FTaskData
 	{
-		FName UniqueID;
+		FString UniqueID;
 		int32 SourceIndex;
 		UInterchangeBaseNode* Node;
-		TArray<FName> Dependencies;
+		TArray<FString> Dependencies;
 		FGraphEventRef GraphEventRef;
 		FGraphEventArray Prerequistes;
 		UInterchangeFactoryBase* Factory;
@@ -60,7 +60,7 @@ void UE::Interchange::FTaskParsing::DoTask(ENamedThreads::Type CurrentThread, co
 			{
 				continue;
 			}
-			BaseNodeContainer->IterateNodes([&](const FName NodeUID, UInterchangeBaseNode* Node)
+			BaseNodeContainer->IterateNodes([&](const FString& NodeUID, UInterchangeBaseNode* Node)
 			{
 				if (Node->GetAssetClass() != nullptr)
 				{
