Add a generic mesh pipeline for Interchange, it allow to add meshes import options

diff --git a/Engine/Config/BaseEngine.ini b/Engine/Config/BaseEngine.ini
--- a/Engine/Config/BaseEngine.ini
+++ b/Engine/Config/BaseEngine.ini
@@ -2550,4 +2550,7 @@ UnrealBuildTool=Engine/Binaries/DotNET/UnrealBuildTool/UnrealBuildTool.exe
 
 [LiveCoding]
 bDisableActionLimit=false
-ActionLimit=100
\ No newline at end of file
+ActionLimit=100
+
+[/Script/InterchangeEngine.InterchangeProjectSettings]
+PipelineStack=/Script/InterchangePipelinePlugin.InterchangeGenericMeshPipeline
\ No newline at end of file
diff --git a/Engine/Plugins/Enterprise/DatasmithImporter/Source/DatasmithImporter/Private/DatasmithTextureImporter.cpp b/Engine/Plugins/Enterprise/DatasmithImporter/Source/DatasmithImporter/Private/DatasmithTextureImporter.cpp
--- a/Engine/Plugins/Enterprise/DatasmithImporter/Source/DatasmithImporter/Private/DatasmithTextureImporter.cpp
+++ b/Engine/Plugins/Enterprise/DatasmithImporter/Source/DatasmithImporter/Private/DatasmithTextureImporter.cpp
@@ -33,6 +33,7 @@
 #include "InterchangeManager.h"
 #include "ObjectTools.h"
 #include "RHI.h"
+#include "InterchangeSourceData.h"
 #include "InterchangeTextureNode.h"
 
 #define LOCTEXT_NAMESPACE "DatasmithTextureImport"
@@ -380,7 +381,7 @@ UE::Interchange::FAssetImportResultRef FDatasmithTextureImporter::CreateTextureA
 	return UInterchangeManager::GetInterchangeManager().ImportAssetAsync( ContentPath, ScopedSourceData.GetSourceData(), ImportAssetParameters );
 }
 
-bool UDatasmithTexturePipeline::ExecutePreImportPipeline(UInterchangeBaseNodeContainer* BaseNodeContainer)
+bool UDatasmithTexturePipeline::ExecutePreImportPipeline(UInterchangeBaseNodeContainer* BaseNodeContainer, const TArray<UInterchangeSourceData*>& SourceDatas)
 {
 	if ( !TextureElement.IsValid() )
 	{
diff --git a/Engine/Plugins/Enterprise/DatasmithImporter/Source/DatasmithImporter/Public/DatasmithTextureImporter.h b/Engine/Plugins/Enterprise/DatasmithImporter/Source/DatasmithImporter/Public/DatasmithTextureImporter.h
--- a/Engine/Plugins/Enterprise/DatasmithImporter/Source/DatasmithImporter/Public/DatasmithTextureImporter.h
+++ b/Engine/Plugins/Enterprise/DatasmithImporter/Source/DatasmithImporter/Public/DatasmithTextureImporter.h
@@ -9,6 +9,7 @@
 
 #include "InterchangeEngineFwd.h"
 #include "InterchangePipelineBase.h"
+#include "InterchangeSourceData.h"
 #include "Async/TaskGraphInterfaces.h"
 
 #include "DatasmithTextureImporter.generated.h"
@@ -45,7 +46,7 @@ class UDatasmithTexturePipeline : public UInterchangePipelineBase
 	GENERATED_BODY()
 
 public:
-	virtual bool ExecutePreImportPipeline(UInterchangeBaseNodeContainer* BaseNodeContainer) override;
+	virtual bool ExecutePreImportPipeline(UInterchangeBaseNodeContainer* BaseNodeContainer, const TArray<UInterchangeSourceData*>& SourceDatas) override;
 
 	bool IsScripted() override
 	{
diff --git a/Engine/Plugins/Runtime/Interchange/Interchange.uplugin b/Engine/Plugins/Runtime/Interchange/Interchange.uplugin
--- a/Engine/Plugins/Runtime/Interchange/Interchange.uplugin
+++ b/Engine/Plugins/Runtime/Interchange/Interchange.uplugin
@@ -30,6 +30,11 @@
 			"Type": "Runtime",
 			"LoadingPhase": "Default"
 		},
+		{
+			"Name": "InterchangePipelinePlugin",
+			"Type": "Runtime",
+			"LoadingPhase": "Default"
+		},
 		{
 			"Name": "InterchangeDispatcher",
 			"Type": "Runtime",
diff --git a/Engine/Plugins/Runtime/Interchange/Source/Import/Private/Mesh/InterchangeSkeletalMeshFactory.cpp b/Engine/Plugins/Runtime/Interchange/Source/Import/Private/Mesh/InterchangeSkeletalMeshFactory.cpp
--- a/Engine/Plugins/Runtime/Interchange/Source/Import/Private/Mesh/InterchangeSkeletalMeshFactory.cpp
+++ b/Engine/Plugins/Runtime/Interchange/Source/Import/Private/Mesh/InterchangeSkeletalMeshFactory.cpp
@@ -1300,6 +1300,7 @@ void UInterchangeSkeletalMeshFactory::PostImportGameThreadCallback(const FPostIm
 																										  , Arguments.NodeUniqueID
 																										  , Arguments.NodeContainer);
 		UE::Interchange::FFactoryCommon::UpdateImportAssetData(UpdateImportAssetDataParameters);
+		SkeletalMesh->SetAssetImportData(ImportDataPtr);
 	}
 #endif
 }
\ No newline at end of file
diff --git a/Engine/Plugins/Runtime/Interchange/Source/Pipelines/InterchangePipelinePlugin.Build.cs b/Engine/Plugins/Runtime/Interchange/Source/Pipelines/InterchangePipelinePlugin.Build.cs
new file mode 100644
--- /dev/null
+++ b/Engine/Plugins/Runtime/Interchange/Source/Pipelines/InterchangePipelinePlugin.Build.cs
@@ -0,0 +1,22 @@
+// Copyright Epic Games, Inc. All Rights Reserved.
+
+namespace UnrealBuildTool.Rules
+{
+	public class InterchangePipelinePlugin : ModuleRules
+	{
+		public InterchangePipelinePlugin(ReadOnlyTargetRules Target) : base(Target)
+		{
+			PublicDependencyModuleNames.AddRange(
+				new string[]
+				{
+					"Core",
+					"CoreUObject",
+					"Engine",
+					"InterchangeCore",
+					"InterchangeEngine",
+					"InterchangeNodePlugin"
+				}
+				);
+		}
+	}
+}
diff --git a/Engine/Plugins/Runtime/Interchange/Source/Pipelines/Private/InterchangeGenericMeshPipeline.cpp b/Engine/Plugins/Runtime/Interchange/Source/Pipelines/Private/InterchangeGenericMeshPipeline.cpp
new file mode 100644
--- /dev/null
+++ b/Engine/Plugins/Runtime/Interchange/Source/Pipelines/Private/InterchangeGenericMeshPipeline.cpp
@@ -0,0 +1,148 @@
+// Copyright Epic Games, Inc. All Rights Reserved. 
+#include "InterchangeGenericMeshPipeline.h"
+
+#include "CoreMinimal.h"
+#include "InterchangeSkeletalMeshLodDataNode.h"
+#include "InterchangeSkeletalMeshNode.h"
+#include "InterchangeSkeletonNode.h"
+#include "InterchangeSourceData.h"
+#include "Misc/Paths.h"
+#include "Nodes/InterchangeBaseNodeContainer.h"
+#include "UObject/Object.h"
+#include "UObject/ObjectMacros.h"
+
+bool UInterchangeGenericMeshPipeline::ExecutePreImportPipeline(UInterchangeBaseNodeContainer* BaseNodeContainer, const TArray<UInterchangeSourceData*>& SourceDatas)
+{
+	PreImportPipelineHelper.Initialize(BaseNodeContainer, SourceDatas, this);
+
+	//Make sure all node graph changes, merges, etc have been done before calling this feature since the rename is done only if there is one imported asset.
+	PreImportPipelineHelper.ChangeAssetNameToUseFilename();
+	
+	return true;
+}
+
+void UInterchangeGenericMeshPipeline::FPreImportPipelineHelper::Initialize(UInterchangeBaseNodeContainer* InBaseNodeContainer, const TArray<UInterchangeSourceData*>& InSourceDatas, UInterchangeGenericMeshPipeline* InPipelineOwner)
+{
+	BaseNodeContainer = InBaseNodeContainer;
+	SourceDatas = InSourceDatas;
+	PipelineOwner = InPipelineOwner;
+
+	RefreshMeshNodes();
+}
+
+bool UInterchangeGenericMeshPipeline::FPreImportPipelineHelper::IsValid() const 
+{
+	//Check all the initialized data
+	if (!BaseNodeContainer)
+	{
+		return false;
+	}
+	
+	for (const UInterchangeSourceData* SourceData : SourceDatas)
+	{
+		if (!SourceData)
+		{
+			return false;
+		}
+	}
+	
+	if (!PipelineOwner)
+	{
+		return false;
+	}
+
+	return true;
+}
+
+void UInterchangeGenericMeshPipeline::FPreImportPipelineHelper::RefreshMeshNodes()
+{
+	if (!IsValid())
+	{
+		return;
+	}
+	//Get all node we want this pipeline interact with.
+	BaseNodeContainer->IterateNodes([this](const FString& NodeUID, UInterchangeBaseNode* Node)
+	{
+		if (UInterchangeSkeletalMeshNode* SkeletalMeshNode = Cast<UInterchangeSkeletalMeshNode>(Node))
+		{
+			SkeletalMeshNodes.Add(SkeletalMeshNode);
+		}
+	});
+}
+
+void UInterchangeGenericMeshPipeline::FPreImportPipelineHelper::ChangeAssetNameToUseFilename()
+{
+	if (!IsValid())
+	{
+		return;
+	}
+
+	//Is the feature enable or not
+	if (!PipelineOwner->bUseFilenameToNameMeshes)
+	{
+		return;
+	}
+
+	if (SourceDatas.Num() != 1 || !(SourceDatas[0]))
+	{
+		return;
+	}
+
+	int32 MeshCount = 0;
+	UInterchangeBaseNode* NodeToRename = nullptr;
+	//Add all enable mesh count here, The Use filename has name feature work only if there is one asset to import
+	//Start by the skeletalmesh
+	for (int32 SkeletalMeshIndex = 0; SkeletalMeshIndex < SkeletalMeshNodes.Num(); ++SkeletalMeshIndex)
+	{
+		if (SkeletalMeshNodes[SkeletalMeshIndex]->IsEnabled())
+		{
+			MeshCount++;
+			if (MeshCount == 1)
+			{
+				NodeToRename = SkeletalMeshNodes[SkeletalMeshIndex];
+			}
+			else
+			{
+				NodeToRename = nullptr;
+			}
+		}
+	}
+	//Iterate staticmesh
+
+	if (!NodeToRename)
+	{
+		return;
+	}
+
+	//If the NodeToRename pointer is valid we should have only one mesh count
+	if (!ensure(MeshCount == 1))
+	{
+		return;
+	}
+
+	FString SourceFilename = SourceDatas[0]->GetFilename();
+	FString OverrideMeshName = FPaths::GetBaseFilename(SourceFilename);
+	NodeToRename->SetDisplayLabel(OverrideMeshName);
+
+	//For skeletalmesh we need to rename also the skeleton asset node
+	if (UInterchangeSkeletalMeshNode* SkeletalMeshNode = Cast<UInterchangeSkeletalMeshNode>(NodeToRename))
+	{
+		TArray<FString> LodDataUniqueIds;
+		SkeletalMeshNode->GetLodDataUniqueIds(LodDataUniqueIds);
+		if (LodDataUniqueIds.Num() > 0)
+		{
+			//Get the LOD 0 node
+			if (UInterchangeSkeletalMeshLodDataNode* SkeletalMeshLodDataNode = Cast<UInterchangeSkeletalMeshLodDataNode>(BaseNodeContainer->GetNode(LodDataUniqueIds[0])))
+			{
+				FString SkeletonNodeId;
+				if (SkeletalMeshLodDataNode->GetCustomSkeletonID(SkeletonNodeId))
+				{
+					if (UInterchangeSkeletonNode* SkeletonNode = Cast<UInterchangeSkeletonNode>(BaseNodeContainer->GetNode(SkeletonNodeId)))
+					{
+						SkeletonNode->SetDisplayLabel(OverrideMeshName + TEXT("_Skeleton"));
+					}
+				}
+			}
+		}
+	}
+}
\ No newline at end of file
diff --git a/Engine/Plugins/Runtime/Interchange/Source/Pipelines/Private/InterchangePipelinePlugin.cpp b/Engine/Plugins/Runtime/Interchange/Source/Pipelines/Private/InterchangePipelinePlugin.cpp
new file mode 100644
--- /dev/null
+++ b/Engine/Plugins/Runtime/Interchange/Source/Pipelines/Private/InterchangePipelinePlugin.cpp
@@ -0,0 +1,49 @@
+// Copyright Epic Games, Inc. All Rights Reserved.
+
+#include "CoreMinimal.h"
+#include "Engine/Engine.h"
+#include "IInterchangePipelinePlugin.h"
+#include "InterchangeGenericMeshPipeline.h"
+#include "InterchangeManager.h"
+#include "Misc/CoreDelegates.h"
+#include "Modules/ModuleManager.h"
+
+
+class FInterchangePipelinePlugin : public IInterchangePipelinePlugin
+{
+	/** IModuleInterface implementation */
+	virtual void StartupModule() override;
+	virtual void ShutdownModule() override;
+};
+
+IMPLEMENT_MODULE( FInterchangePipelinePlugin, InterchangePipelinePlugin)
+
+
+
+void FInterchangePipelinePlugin::StartupModule()
+{
+	auto RegisterItems = []()
+	{
+		UInterchangeManager& InterchangeManager = UInterchangeManager::GetInterchangeManager();
+		//Register the Pipeline
+		//InterchangeManager.RegisterPipeline(UInterchangeGenericMeshPipeline::StaticClass());
+	};
+
+	if (GEngine)
+	{
+		RegisterItems();
+	}
+	else
+	{
+		FCoreDelegates::OnPostEngineInit.AddLambda(RegisterItems);
+	}
+}
+
+
+void FInterchangePipelinePlugin::ShutdownModule()
+{
+	
+}
+
+
+
diff --git a/Engine/Plugins/Runtime/Interchange/Source/Pipelines/Public/IInterchangePipelinePlugin.h b/Engine/Plugins/Runtime/Interchange/Source/Pipelines/Public/IInterchangePipelinePlugin.h
new file mode 100644
--- /dev/null
+++ b/Engine/Plugins/Runtime/Interchange/Source/Pipelines/Public/IInterchangePipelinePlugin.h
@@ -0,0 +1,39 @@
+// Copyright Epic Games, Inc. All Rights Reserved.
+
+#pragma once
+
+#include "CoreMinimal.h"
+#include "Modules/ModuleInterface.h"
+#include "Modules/ModuleManager.h"
+
+
+/**
+ * The public interface to this module
+ */
+class IInterchangePipelinePlugin : public IModuleInterface
+{
+
+public:
+
+	/**
+	 * Singleton-like access to this module's interface.  This is just for convenience!
+	 * Beware of calling this during the shutdown phase, though.  Your module might have been unloaded already.
+	 *
+	 * @return Returns singleton instance, loading the module on demand if needed
+	 */
+	static inline IInterchangePipelinePlugin& Get()
+	{
+		return FModuleManager::LoadModuleChecked< IInterchangePipelinePlugin >( "InterchangePipelinePlugin" );
+	}
+
+	/**
+	 * Checks to see if this module is loaded and ready.  It is only valid to call Get() if IsAvailable() returns true.
+	 *
+	 * @return True if the module is loaded and ready to use
+	 */
+	static inline bool IsAvailable()
+	{
+		return FModuleManager::Get().IsModuleLoaded( "InterchangePipelinePlugin" );
+	}
+};
+
diff --git a/Engine/Plugins/Runtime/Interchange/Source/Pipelines/Public/InterchangeGenericMeshPipeline.h b/Engine/Plugins/Runtime/Interchange/Source/Pipelines/Public/InterchangeGenericMeshPipeline.h
new file mode 100644
--- /dev/null
+++ b/Engine/Plugins/Runtime/Interchange/Source/Pipelines/Public/InterchangeGenericMeshPipeline.h
@@ -0,0 +1,93 @@
+// Copyright Epic Games, Inc. All Rights Reserved.
+
+#pragma once
+
+#include "CoreMinimal.h"
+#include "InterchangePipelineBase.h"
+#include "InterchangeSourceData.h"
+#include "Nodes/InterchangeBaseNodeContainer.h"
+#include "UObject/Object.h"
+#include "UObject/ObjectMacros.h"
+
+#include "InterchangeGenericMeshPipeline.generated.h"
+
+class UInterchangeSkeletalMeshNode;
+
+/**
+ * This pipeline is the generic pipeline option for all meshes type and should be call before specialized Mesh pipeline (like generic static mesh or skeletal mesh pipelines)
+ * All shared import options between mesh type should be added here.
+ *
+ * UPROPERTY possible meta values:
+ * @meta ImportOnly - Boolean, the property is use only when we import not when we re-import. Cannot be mix with ReimportOnly!
+ * @meta ReimportOnly - Boolean, the property is use only when we re-import not when we import. Cannot be mix with ImportOnly!
+ * @meta MeshType - String, the property is for static or skeletal or both (static | skeletal) mesh type. If not specified it will apply to all mesh type.
+ */
+UCLASS(BlueprintType)
+class INTERCHANGEPIPELINEPLUGIN_API UInterchangeGenericMeshPipeline : public UInterchangePipelineBase
+{
+	GENERATED_BODY()
+
+public:
+	//////////////////////////////////////////////////////////////////////////
+	// BEGIN Pre import pipeline properties
+
+	/** If enable and there is only one mesh in the imported source file, the filename will be use to name the asset */
+	UPROPERTY(EditAnywhere, BlueprintReadWrite, Category = "PreImport | Miscellaneous", meta = (ImportOnly = "true", MeshType = "Static | Skeletal"))
+	bool bUseFilenameToNameMeshes = true;
+
+	// END Pre import pipeline properties
+	//////////////////////////////////////////////////////////////////////////
+
+protected:
+	
+	virtual bool ExecutePreImportPipeline(UInterchangeBaseNodeContainer* BaseNodeContainer, const TArray<UInterchangeSourceData*>& SourceDatas) override;
+
+	//virtual bool ExecutePostImportPipeline(const UInterchangeBaseNodeContainer* BaseNodeContainer, const FName& NodeKey, UObject* CreatedAsset) override;
+
+	virtual bool CanExecuteOnAnyThread(EInterchangePipelineTask PipelineTask) override
+	{
+		//If a blueprint or python derived from this class, it will be execute on the game thread since we cannot currently execute script outside of the game thread, even if this function return true.
+		return true;
+	}
+
+	//virtual bool ExecuteExportPipeline(UInterchangeBaseNodeContainer* BaseNodeContainer) override;
+private:
+
+	struct FPreImportPipelineHelper
+	{
+	public:
+		/** This function initialize the helper with the proper data*/
+		void Initialize(UInterchangeBaseNodeContainer* InBaseNodeContainer, const TArray<UInterchangeSourceData*>& InSourceDatas, UInterchangeGenericMeshPipeline* InPipelineOwner);
+
+		/**
+		 * This function will change the meshes node display label and use the source filename to build a new mesh name
+		 * @note - If bUseFilenameToNameMeshes is false this function will early exit and not do any name change.
+		 * @note - It will rename staticmesh, skeletalmesh and skeleton nodes
+		 */
+		void ChangeAssetNameToUseFilename();
+
+	private:
+
+		/** This function check all pointer and make sure we can use the pre import pipeline helper functionalities*/
+		bool IsValid() const;
+
+		/**
+		 * This function refresh the skeletal mesh node array to reflect the BaseNodeContainer content
+		 * @note - Any pipeline feature that change the topologies of the graph should call this refresh after changing it.
+		 */
+		void RefreshMeshNodes();
+
+
+		UInterchangeBaseNodeContainer* BaseNodeContainer;
+		TArray<UInterchangeSourceData*> SourceDatas;
+		UInterchangeGenericMeshPipeline* PipelineOwner;
+
+		/** List of all SkeletalMeshNodes, must be refresh if we change the node graph(add/remove/move any node in the graph) */
+		TArray<UInterchangeSkeletalMeshNode*> SkeletalMeshNodes;
+	};
+
+	/** Cache the PreImportPipeline parameter in this structure to avoid passing parameter to each private functions */
+	FPreImportPipelineHelper PreImportPipelineHelper;
+};
+
+
diff --git a/Engine/Source/Runtime/Interchange/Core/Private/Nodes/InterchangeBaseNodeContainer.cpp b/Engine/Source/Runtime/Interchange/Core/Private/Nodes/InterchangeBaseNodeContainer.cpp
--- a/Engine/Source/Runtime/Interchange/Core/Private/Nodes/InterchangeBaseNodeContainer.cpp
+++ b/Engine/Source/Runtime/Interchange/Core/Private/Nodes/InterchangeBaseNodeContainer.cpp
@@ -71,6 +71,17 @@ void UInterchangeBaseNodeContainer::GetRoots(TArray<FString>& RootNodes)
 	}
 }
 
+void UInterchangeBaseNodeContainer::GetNodes(UClass* ClassNode, TArray<FString>& ClassNodes)
+{
+	IterateNodes([&ClassNode, &ClassNodes](const FString& NodeUID, UInterchangeBaseNode* Node)
+	{
+		if(Node->GetClass()->IsChildOf(ClassNode))
+		{
+			ClassNodes.Add(Node->GetUniqueID());
+		}
+	});
+}
+
 UInterchangeBaseNode* UInterchangeBaseNodeContainer::GetNode(const FString& NodeUniqueID)
 {
 	if (NodeUniqueID == UInterchangeBaseNode::InvalidNodeUID())
diff --git a/Engine/Source/Runtime/Interchange/Core/Public/InterchangePipelineBase.h b/Engine/Source/Runtime/Interchange/Core/Public/InterchangePipelineBase.h
--- a/Engine/Source/Runtime/Interchange/Core/Public/InterchangePipelineBase.h
+++ b/Engine/Source/Runtime/Interchange/Core/Public/InterchangePipelineBase.h
@@ -30,12 +30,12 @@ public:
 	 * the Interchange manager is calling this function not the virtual one that is call by the default implementation.
 	 */
 	UFUNCTION(BlueprintNativeEvent, BlueprintCallable, Category = "Interchange | Translator")
-	bool ScriptedExecutePreImportPipeline(UInterchangeBaseNodeContainer* BaseNodeContainer);
+	bool ScriptedExecutePreImportPipeline(UInterchangeBaseNodeContainer* BaseNodeContainer, const TArray<UInterchangeSourceData*>& SourceDatas);
 	/** The default implementation (call if the blueprint do not have any implementation) will call the virtual ExecuteImportPipeline */
-	bool ScriptedExecutePreImportPipeline_Implementation(UInterchangeBaseNodeContainer* BaseNodeContainer)
+	bool ScriptedExecutePreImportPipeline_Implementation(UInterchangeBaseNodeContainer* BaseNodeContainer, const TArray<UInterchangeSourceData*>& SourceDatas)
 	{
 		//By default we call the virtual import pipeline execution
-		return ExecutePreImportPipeline(BaseNodeContainer);
+		return ExecutePreImportPipeline(BaseNodeContainer, SourceDatas);
 	}
 
 	/**
@@ -96,7 +96,7 @@ protected:
 	 * The interchange manager is not calling this function directly. It is calling the blueprint native event in case this object is a blueprint derive object.
 	 * By default the scripted implementation is calling this virtual pipeline.
 	 */
-	virtual bool ExecutePreImportPipeline(UInterchangeBaseNodeContainer* BaseNodeContainer)
+	virtual bool ExecutePreImportPipeline(UInterchangeBaseNodeContainer* BaseNodeContainer, const TArray<UInterchangeSourceData*>& SourceDatas)
 	{
 		return false;
 	}
diff --git a/Engine/Source/Runtime/Interchange/Core/Public/Nodes/InterchangeBaseNodeContainer.h b/Engine/Source/Runtime/Interchange/Core/Public/Nodes/InterchangeBaseNodeContainer.h
--- a/Engine/Source/Runtime/Interchange/Core/Public/Nodes/InterchangeBaseNodeContainer.h
+++ b/Engine/Source/Runtime/Interchange/Core/Public/Nodes/InterchangeBaseNodeContainer.h
@@ -46,6 +46,10 @@ public:
 	UFUNCTION(BlueprintCallable, Category = "Interchange | Node Container")
 	void GetRoots(TArray<FString>& RootNodes);
 
+	/** Return all nodes that are of the ClassNode type*/
+	UFUNCTION(BlueprintCallable, Category = "Interchange | Node Container")
+	void GetNodes(UClass* ClassNode, TArray<FString>& ClassNodes);
+
 	/** Get an node pointer */
 	UFUNCTION(BlueprintCallable, Category = "Interchange | Node Container")
 	UInterchangeBaseNode* GetNode(const FString& NodeUniqueID);
diff --git a/Engine/Source/Runtime/Interchange/Core/Public/Types/AttributeStorage.h b/Engine/Source/Runtime/Interchange/Core/Public/Types/AttributeStorage.h
--- a/Engine/Source/Runtime/Interchange/Core/Public/Types/AttributeStorage.h
+++ b/Engine/Source/Runtime/Interchange/Core/Public/Types/AttributeStorage.h
@@ -874,15 +874,16 @@ namespace UE
 					{
 						return EAttributeStorageResult::Operation_Error_CannotFoundKey;
 					}
+					const FAttributeKey OriginalKey = *ElementAttributeKeyPtr;
 					const EAttributeProperty AttributePropertyBackup = AttributeAllocationInfo->Property;
-					const EAttributeStorageResult RemoveResult = UnregisterAttribute(*ElementAttributeKeyPtr);
+					const EAttributeStorageResult RemoveResult = UnregisterAttribute(OriginalKey);
 					if (!IsAttributeStorageResultSuccess(RemoveResult))
 					{
 						return RemoveResult;
 					}
 					AttributeAllocationInfo = nullptr;
 					//Reallocate a new one
-					AttributeAllocationInfo = &AttributeAllocationTable.Add(*ElementAttributeKeyPtr);
+					AttributeAllocationInfo = &AttributeAllocationTable.Add(OriginalKey);
 					AttributeAllocationInfo->Type = ValueType;
 					AttributeAllocationInfo->Offset = AttributeStorage.AddZeroed(ValueSize);
 					//Force the specified attribute property
diff --git a/Engine/Source/Runtime/Interchange/Engine/Private/Tasks/InterchangeTaskPipeline.cpp b/Engine/Source/Runtime/Interchange/Engine/Private/Tasks/InterchangeTaskPipeline.cpp
--- a/Engine/Source/Runtime/Interchange/Engine/Private/Tasks/InterchangeTaskPipeline.cpp
+++ b/Engine/Source/Runtime/Interchange/Engine/Private/Tasks/InterchangeTaskPipeline.cpp
@@ -33,7 +33,7 @@ void UE::Interchange::FTaskPipelinePreImport::DoTask(ENamedThreads::Type Current
 
 			if (ensure(AsyncHelper->BaseNodeContainers[GraphIndex].IsValid()))
 			{
-				Pipeline->ScriptedExecutePreImportPipeline(AsyncHelper->BaseNodeContainers[GraphIndex].Get());
+				Pipeline->ScriptedExecutePreImportPipeline(AsyncHelper->BaseNodeContainers[GraphIndex].Get(), AsyncHelper->SourceDatas);
 			}
 		}
 	}
