Adding UDIM support to the interchange framework.

I used a confidential dataset to benchmark the performence. 112 png files weighting 3.2 GB.

The old import system took 75.96 seconds to import the data while blocking the gamethread.

The Interchange system took 15.06 seconds to import the data whitout blocking the game thread, but my initial testing found that it resquested around 200 MB of memory then the old system. It's possible that this is due to the memory pool of each thread (it need investigation before jumping to that conclusion).

Those test where made using a machine with a Xeon E5-2643 (6 cores/ 12 Threads)

diff --git a/Engine/Plugins/Runtime/Interchange/Source/Import/InterchangeImportPlugin.Build.cs b/Engine/Plugins/Runtime/Interchange/Source/Import/InterchangeImportPlugin.Build.cs
--- a/Engine/Plugins/Runtime/Interchange/Source/Import/InterchangeImportPlugin.Build.cs
+++ b/Engine/Plugins/Runtime/Interchange/Source/Import/InterchangeImportPlugin.Build.cs
@@ -25,9 +25,10 @@ namespace UnrealBuildTool.Rules
 				new string[]
 				{
 					"ImageWrapper",
-					"RHI",
 					"InterchangeDispatcher",
 					"Json",
+					"RHI",
+					"TextureUtilitiesCommon",
 				}
 				);
 
diff --git a/Engine/Plugins/Runtime/Interchange/Source/Import/Private/InterchangeImportPlugin.cpp b/Engine/Plugins/Runtime/Interchange/Source/Import/Private/InterchangeImportPlugin.cpp
--- a/Engine/Plugins/Runtime/Interchange/Source/Import/Private/InterchangeImportPlugin.cpp
+++ b/Engine/Plugins/Runtime/Interchange/Source/Import/Private/InterchangeImportPlugin.cpp
@@ -1,8 +1,9 @@
 // Copyright Epic Games, Inc. All Rights Reserved.
 
+#include "IInterchangeImportPlugin.h"
+
 #include "CoreMinimal.h"
 #include "Fbx/InterchangeFbxTranslator.h"
-#include "IInterchangeImportPlugin.h"
 #include "InterchangeManager.h"
 #include "Material/InterchangeMaterialFactory.h"
 #include "Mesh/InterchangeSkeletalMeshFactory.h"
@@ -18,6 +19,7 @@
 #include "Texture/InterchangePSDTranslator.h"
 #include "Texture/InterchangeTextureFactory.h"
 #include "Texture/InterchangeTGATranslator.h"
+#include "Texture/InterchangeUDIMTranslator.h"
 
 DEFINE_LOG_CATEGORY(LogInterchangeImportPlugin);
 
@@ -41,7 +43,13 @@ void FInterchangeImportPlugin::StartupModule()
 		//Register the translators
 		//Scenes
 		InterchangeManager.RegisterTranslator(UInterchangeFbxTranslator::StaticClass()); //Do not submit uncommented until we replace completly fbx importer (staticmesh + skeletalMesh + animation)
+
 		//Textures
+
+		// UDIM must be registered before the other texture translators
+		InterchangeManager.RegisterTranslator(UInterchangeUDIMTranslator::StaticClass());
+
+
 		InterchangeManager.RegisterTranslator(UInterchangeBMPTranslator::StaticClass());
 		InterchangeManager.RegisterTranslator(UInterchangeDDSTranslator::StaticClass());
 		InterchangeManager.RegisterTranslator(UInterchangeEXRTranslator::StaticClass());
diff --git a/Engine/Plugins/Runtime/Interchange/Source/Import/Private/Texture/InterchangeTextureFactory.cpp b/Engine/Plugins/Runtime/Interchange/Source/Import/Private/Texture/InterchangeTextureFactory.cpp
--- a/Engine/Plugins/Runtime/Interchange/Source/Import/Private/Texture/InterchangeTextureFactory.cpp
+++ b/Engine/Plugins/Runtime/Interchange/Source/Import/Private/Texture/InterchangeTextureFactory.cpp
@@ -9,6 +9,7 @@
 #include "InterchangeTranslatorBase.h"
 #include "LogInterchangeImportPlugin.h"
 #include "Nodes/InterchangeBaseNode.h"
+#include "Texture/InterchangeBlockedTexturePayloadInterface.h"
 #include "Texture/InterchangeTexturePayloadInterface.h"
 #include "TextureCompiler.h"
 
@@ -101,26 +102,31 @@ UObject* UInterchangeTextureFactory::CreateAsset(const UInterchangeTextureFactor
 		return nullptr;
 	}
 
-	const IInterchangeTexturePayloadInterface* TextureTranslator = Cast<IInterchangeTexturePayloadInterface>(Arguments.Translator);
-	if (!TextureTranslator)
+	const TOptional<FString>& PayLoadKey = TextureNode->GetPayLoadKey();
+	if (!PayLoadKey.IsSet())
 	{
 		return nullptr;
 	}
 
-	//
-	// Generic 2D Image
-	//
-	const TOptional<FString>& PayLoadKey = TextureNode->GetPayLoadKey();
-	if (!PayLoadKey.IsSet())
+	TOptional<UE::Interchange::FImportImage> PayloadData;
+	TOptional<UE::Interchange::FImportBlockedImage> BlockedPayloadData;
+	if (const IInterchangeTexturePayloadInterface* TextureTranslator = Cast<IInterchangeTexturePayloadInterface>(Arguments.Translator))
 	{
-		return nullptr;
+		 PayloadData = TextureTranslator->GetTexturePayloadData(Arguments.SourceData, PayLoadKey.GetValue());
+	}
+	else if (const IInterchangeBlockedTexturePayloadInterface* BlockedTextureTranslator = Cast<IInterchangeBlockedTexturePayloadInterface>(Arguments.Translator))
+	{
+		/**
+		 * Possible improvement store when a file was last modified and use that as heuristic to skip the file assuming that the file is same as the one in the previous import
+		 * (how to deal with engine changes or translator changes/hotfix?)
+		 */ 
+		BlockedPayloadData = BlockedTextureTranslator->GetBlockedTexturePayloadData(TextureNode->GetSourceBlocks(), Arguments.SourceData);
 	}
-	const TOptional<UE::Interchange::FImportImage> PayloadData = TextureTranslator->GetTexturePayloadData(Arguments.SourceData, PayLoadKey.GetValue());
-	if(!PayloadData.IsSet())
+
+	if(!PayloadData.IsSet() && !BlockedPayloadData.IsSet())
 	{
 		return nullptr;
 	}
-	const UE::Interchange::FImportImage& Image = PayloadData.GetValue();
 
 	const UClass* TextureClass = TextureNode->GetAssetClass();
 	check(TextureClass && TextureClass->IsChildOf(GetFactoryClass()));
@@ -152,24 +158,72 @@ UObject* UInterchangeTextureFactory::CreateAsset(const UInterchangeTextureFactor
 	UTexture2D* Texture2D = Cast<UTexture2D>(Texture);
 	if (Texture2D)
 	{
-		Texture2D->Source.Init(
-			Image.SizeX,
-			Image.SizeY,
-			/*NumSlices=*/ 1,
-			Image.NumMips,
-			Image.Format,
-			Image.RawData.GetData()
-		);
-		Texture2D->CompressionSettings = Image.CompressionSettings;
-		Texture2D->SRGB = Image.SRGB;
-		
-		//If the MipGenSettings was set by the translator, we must apply it before the build
-		if (Image.MipGenSettings.IsSet())
+		if (BlockedPayloadData)
+		{
+			//
+			// Blocked Texture (also know as UDIM)
+			//
+			const UE::Interchange::FImportBlockedImage& BlockedImage = BlockedPayloadData.GetValue();
+
+			if (BlockedImage.IsBlockedData())
+			{
+				ETextureSourceFormat Format = BlockedImage.ImagesData[0].Format;
+				TextureCompressionSettings CompressionSettings = BlockedImage.ImagesData[0].CompressionSettings;
+				bool bSRGB = BlockedImage.ImagesData[0].SRGB;
+
+				TArray<const uint8*> SourceImageDatasPtr;
+				SourceImageDatasPtr.Reserve(BlockedImage.ImagesData.Num());
+				for (const  UE::Interchange::FImportImage& Image : BlockedImage.ImagesData)
+				{
+					SourceImageDatasPtr.Add(Image.RawData.GetData());
+				}
+
+				Texture2D->Source.InitBlocked(
+					&Format,
+					BlockedImage.BlocksData.GetData(),
+					/*NumSlices=*/ 1,
+					BlockedImage.BlocksData.Num(),
+					SourceImageDatasPtr.GetData()
+				);
+
+				Texture2D->CompressionSettings = CompressionSettings;
+				Texture2D->SRGB = bSRGB;
+				Texture2D->VirtualTextureStreaming = true;
+			}
+			else
+			{
+				//Import as a normal texture
+				PayloadData = MoveTemp(BlockedPayloadData.GetValue().ImagesData[0]);
+			}
+		}
+
+		if (PayloadData)
 		{
-			// if the source has mips we keep the mips by default, unless the user changes that
-			Texture2D->MipGenSettings = Image.MipGenSettings.GetValue();
+			//
+			// Generic 2D Image
+			//
+			const UE::Interchange::FImportImage& Image = PayloadData.GetValue();
+			Texture2D->Source.Init(
+				Image.SizeX,
+				Image.SizeY,
+				/*NumSlices=*/ 1,
+				Image.NumMips,
+				Image.Format,
+				Image.RawData.GetData()
+			);
+			Texture2D->CompressionSettings = Image.CompressionSettings;
+			Texture2D->SRGB = Image.SRGB;
+
+			//If the MipGenSettings was set by the translator, we must apply it before the build
+			if (Image.MipGenSettings.IsSet())
+			{
+				// if the source has mips we keep the mips by default, unless the user changes that
+				Texture2D->MipGenSettings = Image.MipGenSettings.GetValue();
+			}
+
 		}
 
+
 		if(!Arguments.ReimportObject)
 		{
 			/** Apply all TextureNode custom attributes to the texture asset */
diff --git a/Engine/Plugins/Runtime/Interchange/Source/Import/Private/Texture/InterchangeUDIMTranslator.cpp b/Engine/Plugins/Runtime/Interchange/Source/Import/Private/Texture/InterchangeUDIMTranslator.cpp
new file mode 100644
--- /dev/null
+++ b/Engine/Plugins/Runtime/Interchange/Source/Import/Private/Texture/InterchangeUDIMTranslator.cpp
@@ -0,0 +1,328 @@
+// Copyright Epic Games, Inc. All Rights Reserved.
+
+#include "Texture/InterchangeUDIMTranslator.h"
+
+#include "InterchangeManager.h"
+#include "InterchangeTextureNode.h"
+#include "Texture/InterchangeTexturePayloadInterface.h"
+
+
+#include "UDIMUtilities.h"
+#include "Async/ParallelFor.h"
+
+namespace UE
+{
+	namespace Interchange
+	{
+		namespace Private
+		{
+			namespace InterchangeUDIMTranslator
+			{
+				class FScopedTranslatorsAndSourceData
+				{
+				public:
+					FScopedTranslatorsAndSourceData(UInterchangeTranslatorBase& BaseTranslator, uint32 NumRequired)
+					{
+						UObject* TransientPackage = GetTransientPackage();
+
+						TranslatorsAndSourcesData.Reserve(NumRequired);
+						TranslatorsAndSourcesData.Emplace(&BaseTranslator, NewObject<UInterchangeSourceData>(TransientPackage, UInterchangeSourceData::StaticClass()));
+						
+						UClass* TranslatorClass = BaseTranslator.GetClass();
+						for (uint32 Index = 1; Index < NumRequired; Index++)
+						{
+							TranslatorsAndSourcesData.Emplace(
+								  NewObject<UInterchangeTranslatorBase>(TransientPackage, TranslatorClass, NAME_None)
+								, NewObject<UInterchangeSourceData>(TransientPackage, UInterchangeSourceData::StaticClass(), NAME_None)
+								);
+
+						}
+					};
+
+					~FScopedTranslatorsAndSourceData()
+					{
+						for (TPair<UInterchangeTranslatorBase*,UInterchangeSourceData*>& TranslatorAndSourceData : TranslatorsAndSourcesData)
+						{
+							TranslatorAndSourceData.Key->ClearInternalFlags(EInternalObjectFlags::Async);
+							TranslatorAndSourceData.Value->ClearInternalFlags(EInternalObjectFlags::Async);
+						}
+					};
+
+					TPair<UInterchangeTranslatorBase*,UInterchangeSourceData*>& operator[](uint32 Index)
+					{
+						return TranslatorsAndSourcesData[Index];
+					}
+
+					const TPair<UInterchangeTranslatorBase*,UInterchangeSourceData*>& operator[](uint32 Index) const
+					{
+						return TranslatorsAndSourcesData[Index];
+					}
+
+					uint32 Num() const
+					{
+						return TranslatorsAndSourcesData.Num();
+					}
+
+				private:
+					TArray<TPair<UInterchangeTranslatorBase*,UInterchangeSourceData*>> TranslatorsAndSourcesData;
+				};
+		
+				class FScopedClearAsyncFlag
+				{
+				public:
+					FScopedClearAsyncFlag(UObject* InObject)
+						: Object(nullptr)
+					{
+						if (InObject && InObject->HasAnyInternalFlags(EInternalObjectFlags::Async))
+						{
+							Object = InObject;
+						}
+					}
+
+					~FScopedClearAsyncFlag()
+					{
+						if (Object)
+						{
+							Object->ClearInternalFlags(EInternalObjectFlags::Async);
+						}
+					}
+
+				private:
+					UObject* Object;
+				};
+			}
+		}
+	}
+}
+
+
+UInterchangeUDIMTranslator::UInterchangeUDIMTranslator()
+	: Super()
+	, IInterchangeBlockedTexturePayloadInterface()
+	, UdimRegexPattern( TEXT(R"((.+?)[._](\d{4})$)") )
+{
+}
+
+bool UInterchangeUDIMTranslator::CanImportSourceData(const UInterchangeSourceData* InSourceData) const
+{
+	// 1) Parse the file name for a UDIM pattern
+	const FString CurrentFilename = InSourceData->GetFilename();
+
+	const FString FilenameNoExtension = FPaths::GetBaseFilename(CurrentFilename);
+
+	FString PreUDIMName;
+	FString PostUDIMName;
+	const int32 BaseUDIMIndex = UE::TextureUtilitiesCommon::ParseUDIMName(FilenameNoExtension, UdimRegexPattern, PreUDIMName, PostUDIMName);
+	const FString BaseUDIMName = PreUDIMName + PostUDIMName;
+
+	if (BaseUDIMIndex != INDEX_NONE)
+	{
+		UInterchangeManager& InterchangeManager = UInterchangeManager::GetInterchangeManager();
+
+		// 2) Check if the current the file can be imported
+		if (UInterchangeTranslatorBase* Translator = InterchangeManager.GetTranslatorSupportingPayloadInterfaceForSourceData(InSourceData, UInterchangeTexturePayloadInterface::StaticClass()))
+		{
+			UE::Interchange::Private::InterchangeUDIMTranslator::FScopedClearAsyncFlag AsyncFlagGuard(Translator);
+			// 3) If another file with the pattern can be imported
+			// Filter for other potential UDIM pages, with the same base name and file extension
+			const FString Path = FPaths::GetPath(CurrentFilename);
+			const FString UDIMFilter = (Path / PreUDIMName) + TEXT("*") + PostUDIMName + FPaths::GetExtension(CurrentFilename, true);
+
+			TArray<FString> UDIMFiles;
+			constexpr bool bFindFiles = true;
+			constexpr bool bFindDirectory = false;
+			IFileManager::Get().FindFiles(UDIMFiles, *UDIMFilter, bFindFiles, bFindDirectory);
+
+			for (const FString& UDIMFile : UDIMFiles)
+			{
+				if (!CurrentFilename.EndsWith(UDIMFile))
+				{
+					const int32 UDIMIndex = UE::TextureUtilitiesCommon::ParseUDIMName(FPaths::GetBaseFilename(UDIMFile), UdimRegexPattern, PreUDIMName, PostUDIMName);
+					if (UDIMIndex != BaseUDIMIndex && UDIMIndex != INDEX_NONE)
+					{
+						const FString UDIMName = PreUDIMName + PostUDIMName;
+						if (UDIMName == BaseUDIMName)
+						{
+
+							if (Translator->CanImportSourceData(InSourceData))
+							{
+								return true;
+							}
+						}
+					}
+				}
+			}
+			// If it fail log a warning / info maybe?
+
+		}
+	}
+
+	return false;
+}
+
+bool UInterchangeUDIMTranslator::Translate(const UInterchangeSourceData* SourceData, UInterchangeBaseNodeContainer& BaseNodeContainer) const
+{
+	FString Filename = SourceData->GetFilename();
+	if (!FPaths::FileExists(Filename))
+	{
+		return false;
+	}
+
+	UInterchangeTextureNode* TextureNode = NewObject<UInterchangeTextureNode>(&BaseNodeContainer, NAME_None);
+
+	if (!ensure(TextureNode))
+	{
+		return false;
+	}
+
+	FString PreUDIMName;
+	FString PostUDIMName;
+	UE::TextureUtilitiesCommon::ParseUDIMName(FPaths::GetBaseFilename(Filename), UdimRegexPattern, PreUDIMName, PostUDIMName);
+
+	FString DisplayLabel = PreUDIMName + PostUDIMName;
+
+	TextureNode->InitializeTextureNode(Filename, DisplayLabel, UTexture2D::StaticClass()->GetName());
+	TextureNode->SetPayLoadKey(Filename);
+
+	UInterchangeManager& InterchangeManager = UInterchangeManager::GetInterchangeManager();
+	UInterchangeSourceData* TempSourceData = NewObject<UInterchangeSourceData>(GetTransientPackage(), UInterchangeSourceData::StaticClass());
+	UE::Interchange::Private::InterchangeUDIMTranslator::FScopedClearAsyncFlag AsyncFlagGuard(TempSourceData);
+
+	if (UInterchangeTranslatorBase* Translator = InterchangeManager.GetTranslatorSupportingPayloadInterfaceForSourceData(SourceData, UInterchangeTexturePayloadInterface::StaticClass()))
+	{
+		UE::Interchange::Private::InterchangeUDIMTranslator::FScopedClearAsyncFlag AsyncFlagGuardTranslator(Translator);
+
+		const FString Path = FPaths::GetPath(Filename);
+		const FString UDIMFilter = (Path / PreUDIMName) + TEXT("*") + PostUDIMName + FPaths::GetExtension(Filename, true);
+
+		TMap<int32, FString> UDIMsAndSourcesFile;
+		{
+			TArray<FString> UDIMFiles;
+			IFileManager::Get().FindFiles(UDIMFiles, *UDIMFilter, true, false);
+			UDIMFiles.Reserve(UDIMFiles.Num());
+
+			for (FString& UDIMFile : UDIMFiles)
+			{
+				const int32 UDIMIndex = UE::TextureUtilitiesCommon::ParseUDIMName(FPaths::GetBaseFilename(UDIMFile), UdimRegexPattern, PreUDIMName, PostUDIMName);
+				if (UDIMIndex != INDEX_NONE)
+				{
+					const FString UDIMName = PreUDIMName + PostUDIMName;
+					FString UDIMPath = FPaths::Combine(Path, UDIMFile);
+					TempSourceData->SetFilename(UDIMPath);
+					if (Translator->CanImportSourceData(TempSourceData))
+					{
+						UDIMsAndSourcesFile.Add(UDIMIndex, MoveTemp(UDIMPath));
+					}
+					// else log warning message?
+				}
+			}
+		}
+		TextureNode->SetSourceBlocks(UDIMsAndSourcesFile);
+	}
+
+	BaseNodeContainer.AddNode(TextureNode);
+	
+	return true;
+}
+
+
+TOptional<UE::Interchange::FImportBlockedImage> UInterchangeUDIMTranslator::GetBlockedTexturePayloadData(const TMap<int32, FString>& InSourcesData,  const UInterchangeSourceData* OriginalFile) const
+{
+	const FString CurrentFilename = OriginalFile->GetFilename();
+
+	UInterchangeManager& InterchangeManager = UInterchangeManager::GetInterchangeManager();
+
+	if (UInterchangeTranslatorBase* Translator = InterchangeManager.GetTranslatorSupportingPayloadInterfaceForSourceData(OriginalFile, UInterchangeTexturePayloadInterface::StaticClass()))
+	{
+		UE::Interchange::Private::InterchangeUDIMTranslator::FScopedClearAsyncFlag AsyncFlagGuard(Translator);
+		TArray<const TPair<int32, FString>*> UDIMsAndSourcesFileArray;
+		UDIMsAndSourcesFileArray.Reserve(InSourcesData.Num());
+		for (const TPair<int32, FString>& Pair : InSourcesData)
+		{
+			UDIMsAndSourcesFileArray.Add(&Pair);
+		}
+
+		if (UDIMsAndSourcesFileArray.Num() > 1)
+		{
+			using namespace UE::Interchange::Private::InterchangeUDIMTranslator;
+			FScopedTranslatorsAndSourceData TranslatorsAndSourceData(*Translator, UDIMsAndSourcesFileArray.Num());
+
+			UE::Interchange::FImportBlockedImage BlockedImage;
+			BlockedImage.InitZeroed(UDIMsAndSourcesFileArray.Num());
+
+			TArray<int32> InvalidTextureIndex;
+
+			ParallelFor(UDIMsAndSourcesFileArray.Num(), [&TranslatorsAndSourceData, &UDIMsAndSourcesFileArray, &BlockedImage, &InvalidTextureIndex](int32 Index)
+				{
+					TPair<UInterchangeTranslatorBase*,UInterchangeSourceData*>& TranslatorAndSourceData = TranslatorsAndSourceData[Index];
+					UInterchangeTranslatorBase* Translator = TranslatorAndSourceData.Key;
+					IInterchangeTexturePayloadInterface* TextureTranslator = CastChecked<IInterchangeTexturePayloadInterface>(TranslatorAndSourceData.Key);
+					UInterchangeSourceData* SourceDataForBlock = TranslatorAndSourceData.Value;
+
+					const TPair<int32, FString>& UDIMAndFilename = *UDIMsAndSourcesFileArray[Index];
+					const int32 CurrentUDIM = UDIMAndFilename.Key;
+					SourceDataForBlock->SetFilename(UDIMAndFilename.Value);
+				
+					TOptional<UE::Interchange::FImportImage> Payload;
+					if (Translator->CanImportSourceData(SourceDataForBlock))
+					{
+						Payload = TextureTranslator->GetTexturePayloadData(SourceDataForBlock, SourceDataForBlock->GetFilename());
+					}
+
+					if (Payload.IsSet())
+					{
+						UE::Interchange::FImportImage& Image = Payload.GetValue();
+						// Consume the image
+						BlockedImage.ImagesData[Index] = MoveTemp(Image);
+
+						FTextureSourceBlock& Block =  BlockedImage.BlocksData[Index];
+						Block.BlockX = (CurrentUDIM - 1001) % 10;
+						Block.BlockY = (CurrentUDIM - 1001) / 10;
+						Block.SizeX = Image.SizeX;
+						Block.SizeY = Image.SizeY;
+						Block.NumSlices = 1;
+						Block.NumMips = Image.NumMips;
+					}
+					else
+					{
+						// Todo Interchange Log error
+						InvalidTextureIndex.Add(Index);
+					}
+
+					// Let the translator release it data.
+					Translator->ImportFinish();
+				}
+				, EParallelForFlags::Unbalanced | EParallelForFlags::BackgroundPriority);
+
+			for (int32 Index : InvalidTextureIndex)
+			{
+				BlockedImage.RemoveBlock(Index);
+			}
+
+			if (BlockedImage.HasData())
+			{
+				return BlockedImage;
+			}
+		}
+
+		// Import as a non blocked image
+		if (IInterchangeTexturePayloadInterface* TextureTranslator = CastChecked<IInterchangeTexturePayloadInterface>(Translator))
+		{
+			TOptional<UE::Interchange::FImportImage> Payload = TextureTranslator->GetTexturePayloadData(OriginalFile, OriginalFile->GetFilename());
+			if (Payload.IsSet())
+			{
+				UE::Interchange::FImportBlockedImage BlockedImage;
+				// Consume the image
+				BlockedImage.ImagesData.Add(MoveTemp(Payload.GetValue()));
+				Translator->ImportFinish();
+				return BlockedImage;
+			}
+		}
+
+		Translator->ReleaseSource();
+		// Log error?
+	}
+	return {};
+}
+
+
diff --git a/Engine/Plugins/Runtime/Interchange/Source/Import/Public/Texture/InterchangeBlockedTexturePayloadData.h b/Engine/Plugins/Runtime/Interchange/Source/Import/Public/Texture/InterchangeBlockedTexturePayloadData.h
new file mode 100644
--- /dev/null
+++ b/Engine/Plugins/Runtime/Interchange/Source/Import/Public/Texture/InterchangeBlockedTexturePayloadData.h
@@ -0,0 +1,46 @@
+// Copyright Epic Games, Inc. All Rights Reserved.
+
+#pragma once
+
+#include "Texture/InterchangeTexturePayloadData.h"
+
+#include "CoreMinimal.h"
+#include "Engine/Texture.h"
+
+namespace UE
+{
+	namespace Interchange
+	{
+		// Also known as a UDIMs texture
+		struct FImportBlockedImage
+		{
+			// The first image data will be used for the default compression setting 
+			// The first image will also be use for the pixel format and gamma (bsRGB)
+
+			TArray<FTextureSourceBlock> BlocksData;
+			TArray<FImportImage> ImagesData;
+
+			void InitZeroed(uint32 Count)
+			{
+				BlocksData.AddZeroed(Count);
+				ImagesData.AddZeroed(Count);
+			}
+
+			void RemoveBlock(int32 Index)
+			{
+				BlocksData.RemoveAtSwap(Index);
+				ImagesData.RemoveAtSwap(Index);
+			}
+
+			bool HasData() const
+			{
+				return ImagesData.Num() > 0;
+			}
+
+			bool IsBlockedData() const
+			{
+				return BlocksData.Num() > 0;
+			}
+		};
+	}
+}
\ No newline at end of file
diff --git a/Engine/Plugins/Runtime/Interchange/Source/Import/Public/Texture/InterchangeBlockedTexturePayloadInterface.h b/Engine/Plugins/Runtime/Interchange/Source/Import/Public/Texture/InterchangeBlockedTexturePayloadInterface.h
new file mode 100644
--- /dev/null
+++ b/Engine/Plugins/Runtime/Interchange/Source/Import/Public/Texture/InterchangeBlockedTexturePayloadInterface.h
@@ -0,0 +1,41 @@
+// Copyright Epic Games, Inc. All Rights Reserved.
+
+#pragma once
+
+#include "Texture/InterchangeBlockedTexturePayloadData.h"
+
+#include "CoreMinimal.h"
+#include "UObject/Interface.h"
+#include "UObject/ObjectMacros.h"
+
+#include "InterchangeBlockedTexturePayloadInterface.generated.h"
+
+class UInterchangeSourceData;
+class FString;
+
+UINTERFACE()
+class INTERCHANGEIMPORTPLUGIN_API UInterchangeBlockedTexturePayloadInterface : public UInterface
+{
+	GENERATED_BODY()
+};
+
+/**
+ * Blocked Texture payload interface (Also know as UDIM(s)). Derive from it if your payload can import blocked/UDIMs texture
+ */
+class INTERCHANGEIMPORTPLUGIN_API IInterchangeBlockedTexturePayloadInterface
+{
+	GENERATED_BODY()
+public:
+
+	
+	/**
+	 * Once the translation is done, the import process need a way to retrieve payload data.
+	 * This payload will be use by the factories to create the asset.
+	 *
+	 * @param SourceData - The sources data containing the data to translate
+	 * @param OriginalFile - The file that was used to start the import
+	 * @return a PayloadData containing the import blocked image data. The TOptional will not be set if there is an error.
+	 */
+	virtual TOptional<UE::Interchange::FImportBlockedImage> GetBlockedTexturePayloadData(const TMap<int32, FString>& InSourcesData,  const UInterchangeSourceData* OriginalFile) const = 0;
+};
+
diff --git a/Engine/Plugins/Runtime/Interchange/Source/Import/Public/Texture/InterchangeUDIMTranslator.h b/Engine/Plugins/Runtime/Interchange/Source/Import/Public/Texture/InterchangeUDIMTranslator.h
new file mode 100644
--- /dev/null
+++ b/Engine/Plugins/Runtime/Interchange/Source/Import/Public/Texture/InterchangeUDIMTranslator.h
@@ -0,0 +1,55 @@
+// Copyright Epic Games, Inc. All Rights Reserved.
+
+#pragma once
+
+#include "InterchangeTranslatorBase.h"
+#include "Texture/InterchangeBlockedTexturePayloadData.h"
+#include "Texture/InterchangeBlockedTexturePayloadInterface.h"
+
+#include "CoreMinimal.h"
+#include "UObject/Object.h"
+#include "UObject/ObjectMacros.h"
+
+#include "InterchangeUDIMTranslator.generated.h"
+
+class UInterchangeBaseNodeContainer;
+class UInterchangeSourceData;
+
+UCLASS(BlueprintType)
+class INTERCHANGEIMPORTPLUGIN_API UInterchangeUDIMTranslator : public UInterchangeTranslatorBase, public IInterchangeBlockedTexturePayloadInterface
+{
+	GENERATED_BODY()
+
+public:
+
+	UInterchangeUDIMTranslator();
+
+	/*
+	 * return true if the translator can translate the specified file.
+	 */
+	virtual bool CanImportSourceData(const UInterchangeSourceData* InSourceData) const override;
+
+	/**
+	 * Translate a source data into a node hold by the specified nodes container.
+	 *
+	 * @param SourceData - The source data containing the data to translate
+	 * @param BaseNodeContainer - The unreal objects descriptions container where to put the translated source data.
+	 * @return true if the translator can translate the source data, false otherwise.
+	 */
+	virtual bool Translate(const UInterchangeSourceData* SourceData, UInterchangeBaseNodeContainer& BaseNodeContainer) const override;
+
+	/**
+	 * Once the translation is done, the import process need a way to retrieve payload data.
+	 * This payload will be use by the factories to create the asset.
+	 *
+	 * @param SourceData - The sources data containing the data to translate
+	 * @param PayloadKey - The key to retrieve the a particular payload contain into the specified source data.
+	 * @return a PayloadData containing the import blocked image data. The TOptional will not be set if there is an error.
+	 */
+	virtual TOptional<UE::Interchange::FImportBlockedImage> GetBlockedTexturePayloadData(const TMap<int32, FString>& InSourcesData,  const UInterchangeSourceData* OriginalFile) const override;
+
+	// The pattern used to identify the UDIMs files
+	UPROPERTY(Transient)
+	FString UdimRegexPattern;
+};
+
diff --git a/Engine/Plugins/Runtime/Interchange/Source/Nodes/Public/InterchangeTextureNode.h b/Engine/Plugins/Runtime/Interchange/Source/Nodes/Public/InterchangeTextureNode.h
--- a/Engine/Plugins/Runtime/Interchange/Source/Nodes/Public/InterchangeTextureNode.h
+++ b/Engine/Plugins/Runtime/Interchange/Source/Nodes/Public/InterchangeTextureNode.h
@@ -44,6 +44,12 @@ public:
 #endif
 	}
 
+	virtual void PostInitProperties()
+	{
+		Super::PostInitProperties();
+		SourceBlocks.Initialize(Attributes.ToSharedRef(), TEXT("SourceBlocks"));
+	}
+
 	/**
 	 * Initialize node data
 	 * @param: UniqueID - The uniqueId for this node
@@ -131,7 +137,47 @@ public:
 		}
 	}
 
+public:
+	//////////////////////////////////////////////////////////////////////////
+	// UDIMs begin here
+	// UDIM base texture use a different model for the source data
+	
+	/**
+	 * Get the source blocks for the texture
+	 * If the map is empty then the texture will be simply be imported as normal texture using the payload key
+	 */
+	UFUNCTION(BlueprintCallable, Category = "Interchange | Node | Texture | UDIMs")
+	TMap<int32,FString> GetSourceBlocks() const
+	{
+		return SourceBlocks.ToMap();
+	}
+
+	/**
+	 * Set the source blocks
+	 * Using this will force the texture factory to consider this texture as UDIM.
+	 * @param InSourceBlocks The blocks and their source image that compose the whole texture.
+	 * The textures must be of the same format and use the same pixel format
+	 * The first block inserted in the map be used to determine the accepted texture format and pixel format
+	 */
+	UFUNCTION(BlueprintCallable, Category = "Interchange | Node | Texture | UDIMs")
+	void SetSourceBlocks(const TMap<int32,FString>& InSourceBlocks)
+	{
+		SourceBlocks = InSourceBlocks;
+	}
+
+	/**
+	 * Remove the source block data from texture node
+	 * This is a easy way to stop a texture from being imported as a UDIM
+	 */
+	UFUNCTION(BlueprintCallable, Category = "Interchange | Node | Texture | UDIMs")
+	void ClearSourceBlocksData()
+	{
+		SourceBlocks.Empty();
+	}
 
+	// UDIMs ends here
+	//////////////////////////////////////////////////////////////////////////
+	
 public:
 	/** Return false if the Attribute was not set previously.*/
 	UFUNCTION(BlueprintCallable, Category = "Interchange | Node | Texture")
@@ -937,5 +983,8 @@ protected:
 #if WITH_ENGINE
 	TSubclassOf<UTexture> AssetClass = nullptr;
 #endif
+
+	UE::Interchange::TMapAttributeHelper<int32,FString> SourceBlocks;
+
 	bool bIsTextureNodeClassInitialized = false;
 };
diff --git a/Engine/Source/Editor/UnrealEd/Private/Factories/EditorFactories.cpp b/Engine/Source/Editor/UnrealEd/Private/Factories/EditorFactories.cpp
--- a/Engine/Source/Editor/UnrealEd/Private/Factories/EditorFactories.cpp
+++ b/Engine/Source/Editor/UnrealEd/Private/Factories/EditorFactories.cpp
@@ -283,6 +283,7 @@
 #include "ObjectTools.h"
 
 #include "SkinWeightsUtilities.h"
+#include "UDIMUtilities.h"
 
 DEFINE_LOG_CATEGORY(LogEditorFactories);
 
@@ -3854,45 +3855,6 @@ bool UTextureFactory::DoesSupportClass(UClass* Class)
 	return Class == UTexture2D::StaticClass() || Class == UTextureCube::StaticClass();
 }
 
-static int32 ParseUDIMName(const FString& Name, const FString& UdimRegexPattern, FString& OutPrefixName, FString& OutPostfixName)
-{
-	FRegexPattern RegexPattern( UdimRegexPattern );
-	FRegexMatcher RegexMatcher( RegexPattern, Name );
-
-	int32 UdimValue = INDEX_NONE;
-
-	if ( RegexMatcher.FindNext() )
-	{
-		const int32 StartOfCaptureGroup1 = RegexMatcher.GetCaptureGroupBeginning(1);
-		const int32 EndOfCaptureGroup1 = RegexMatcher.GetCaptureGroupEnding(1);
-		const int32 StartOfCaptureGroup2 = RegexMatcher.GetCaptureGroupBeginning(2);
-		const int32 EndOfCaptureGroup2 = RegexMatcher.GetCaptureGroupEnding(2);
-		const int32 StartOfCaptureGroup3 = RegexMatcher.GetCaptureGroupBeginning(3);
-		const int32 EndOfCaptureGroup3 = RegexMatcher.GetCaptureGroupEnding(3);
-
-		if ( StartOfCaptureGroup1 != INDEX_NONE && StartOfCaptureGroup2 != INDEX_NONE &&
-			 EndOfCaptureGroup1 != INDEX_NONE && EndOfCaptureGroup2 != INDEX_NONE )
-		{
-			LexFromString( UdimValue, *Name.Mid( StartOfCaptureGroup2, EndOfCaptureGroup2 - StartOfCaptureGroup2 ) );
-
-			OutPrefixName = Name.Mid( StartOfCaptureGroup1, EndOfCaptureGroup1 - StartOfCaptureGroup1 );
-
-			if ( StartOfCaptureGroup3 != INDEX_NONE && EndOfCaptureGroup3 != INDEX_NONE )
-			{
-				OutPostfixName = Name.Mid( StartOfCaptureGroup3, EndOfCaptureGroup3 - StartOfCaptureGroup3 );
-			}
-		}
-	}
-	
-	if ( UdimValue < 1001 )
-	{
-		// UDIM starts with 1001 as the origin
-		return INDEX_NONE;
-	}
-
-	return UdimValue;
-}
-
 UObject* UTextureFactory::FactoryCreateBinary
 (
 	UClass*				Class,
@@ -3918,7 +3880,7 @@ UObject* UTextureFactory::FactoryCreateBinary
 
 		FString PreUDIMName;
 		FString PostUDIMName;
-		const int32 BaseUDIMIndex = ParseUDIMName(FilenameNoExtension, UdimRegexPattern, PreUDIMName, PostUDIMName);
+		const int32 BaseUDIMIndex = UE::TextureUtilitiesCommon::ParseUDIMName(FilenameNoExtension, UdimRegexPattern, PreUDIMName, PostUDIMName);
 
 		const FString BaseUDIMName = PreUDIMName + PostUDIMName;
 		if (BaseUDIMIndex != INDEX_NONE)
@@ -3936,7 +3898,7 @@ UObject* UTextureFactory::FactoryCreateBinary
 			{
 				if (!CurrentFilename.EndsWith(UDIMFile) && FactoryCanImport(UDIMFile))
 				{
-					const int32 UDIMIndex = ParseUDIMName(FPaths::GetBaseFilename(UDIMFile), UdimRegexPattern, PreUDIMName, PostUDIMName);
+					const int32 UDIMIndex = UE::TextureUtilitiesCommon::ParseUDIMName(FPaths::GetBaseFilename(UDIMFile), UdimRegexPattern, PreUDIMName, PostUDIMName);
 					if (UDIMIndex != INDEX_NONE)
 					{
 						const FString UDIMName = PreUDIMName + PostUDIMName;
@@ -3962,7 +3924,7 @@ UObject* UTextureFactory::FactoryCreateBinary
 				FString PackageName;
 				InParent->GetName(PackageName);
 
-				const int32 PackageUDIMIndex = ParseUDIMName(PackageName, UdimRegexPattern, PreUDIMName, PostUDIMName);
+				const int32 PackageUDIMIndex = UE::TextureUtilitiesCommon::ParseUDIMName(PackageName, UdimRegexPattern, PreUDIMName, PostUDIMName);
 				const FString PackageUDIMName = PreUDIMName + PostUDIMName;
 
 				const FString ShortPackageName = ObjectTools::SanitizeInvalidChars(BaseUDIMName, INVALID_LONGPACKAGE_CHARACTERS);
@@ -5975,19 +5937,19 @@ bool UReimportFbxSkeletalMeshFactory::CanReimport( UObject* Obj, TArray<FString>
 	USkeletalMesh* SkeletalMesh = Cast<USkeletalMesh>(Obj);
 	if (SkeletalMesh && !SkeletalMesh->HasCustomActorReimportFactory())
 	{
-		if (SkeletalMesh->GetAssetImportData())
+		if (UAssetImportData* AssetImportData = SkeletalMesh->GetAssetImportData())
 		{
-			UFbxAssetImportData *FbxAssetImportData = Cast<UFbxAssetImportData>(SkeletalMesh->GetAssetImportData());
+			UFbxAssetImportData *FbxAssetImportData = Cast<UFbxAssetImportData>(AssetImportData);
 			if (FbxAssetImportData != nullptr && FbxAssetImportData->bImportAsScene)
 			{
 				//This skeletal mesh was import with a scene import, we cannot reimport it here
 				return false;
 			}
-			else if (FPaths::GetExtension(SkeletalMesh->GetAssetImportData()->GetFirstFilename()) == TEXT("abc"))
+			else if (FPaths::GetExtension(AssetImportData->GetFirstFilename()) == TEXT("abc"))
 			{
 				return false;
 			}
-			SkeletalMesh->GetAssetImportData()->ExtractFilenames(OutFilenames);
+			AssetImportData->ExtractFilenames(OutFilenames);
 		}
 		else
 		{
diff --git a/Engine/Source/Editor/UnrealEd/UnrealEd.Build.cs b/Engine/Source/Editor/UnrealEd/UnrealEd.Build.cs
--- a/Engine/Source/Editor/UnrealEd/UnrealEd.Build.cs
+++ b/Engine/Source/Editor/UnrealEd/UnrealEd.Build.cs
@@ -145,8 +145,9 @@ public class UnrealEd : ModuleRules
 				"Projects",
 				"RawMesh",
 				"MeshUtilitiesCommon",
-                "SkeletalMeshUtilitiesCommon",
-                "RenderCore",
+				"SkeletalMeshUtilitiesCommon",
+				"TextureUtilitiesCommon",
+				"RenderCore",
 				"RHI",
 				"Sockets",
 				"SourceControlWindows",
diff --git a/Engine/Source/Runtime/Interchange/Core/Public/Nodes/InterchangeBaseNode.h b/Engine/Source/Runtime/Interchange/Core/Public/Nodes/InterchangeBaseNode.h
--- a/Engine/Source/Runtime/Interchange/Core/Public/Nodes/InterchangeBaseNode.h
+++ b/Engine/Source/Runtime/Interchange/Core/Public/Nodes/InterchangeBaseNode.h
@@ -183,7 +183,7 @@ namespace UE
 				KeyCount = NAME_None;
 			}
 
-			void Initialize(TSharedPtr<FAttributeStorage, ESPMode::ThreadSafe> InAttributes, const FString& BaseKeyName)
+			void Initialize(const TSharedPtr<FAttributeStorage, ESPMode::ThreadSafe>& InAttributes, const FString& BaseKeyName)
 			{
 				Attributes = InAttributes;
 				check(Attributes.IsValid());
@@ -397,6 +397,289 @@ namespace UE
 			}
 		};
 
+		template<typename KeyType, typename ValueType>
+		class TMapAttributeHelper
+		{
+			static_assert(TAttributeTypeTraits<KeyType>::GetType() != EAttributeTypes::None, "The key type must be supported by the attribute storage");
+			static_assert(TAttributeTypeTraits<ValueType>::GetType() != EAttributeTypes::None, "The value type must be supported by the attribute storage");
+
+			template<typename T>
+			using TAttributeHandle = FAttributeStorage::TAttributeHandle<T>;
+
+		public:
+			void Initialize(const TSharedRef<FAttributeStorage, ESPMode::ThreadSafe>& InAttributes, const FString& BaseKeyName)
+			{
+				Attributes = InAttributes;
+				FString BaseTryName = TEXT("__") + BaseKeyName;
+				FAttributeKey KeyCountKey(MoveTemp(BaseTryName));
+				if (!InAttributes->ContainAttribute(KeyCountKey))
+				{ 
+					EAttributeStorageResult Result = InAttributes->RegisterAttribute<int32>(KeyCountKey, 0);
+					check (Result == EAttributeStorageResult::Operation_Success);
+					KeyCountHandle = InAttributes->GetAttributeHandle<int32>(KeyCountKey);
+				}
+				else
+				{
+					KeyCountHandle = InAttributes->GetAttributeHandle<int32>(KeyCountKey);
+
+					//Init map
+					int32 KeyCount;
+					KeyCountHandle.Get(KeyCount);
+
+					CachedKeysAndValues.Reserve(KeyCount);
+
+					for (int32 Index = 0; Index < KeyCount; Index++)
+					{
+						TAttributeHandle<KeyType> KeyAttribute = InAttributes->GetAttributeHandle<KeyType>(GetKeyAttribute(Index));
+						KeyType Key;
+						KeyAttribute.Get(Key);
+						TAttributeHandle<ValueType> ValueAttribute = InAttributes->GetAttributeHandle<ValueType>(GetValueAttribute(Key));
+						CachedKeysAndValues.Add(Key, TPair<TAttributeHandle<KeyType>, TAttributeHandle<ValueType>>(KeyAttribute, ValueAttribute));
+					}
+				}
+			}
+
+			TMapAttributeHelper() = default;
+			TMapAttributeHelper(const TMapAttributeHelper&) = default;
+			TMapAttributeHelper(TMapAttributeHelper&&) = default;
+			TMapAttributeHelper& operator=(const TMapAttributeHelper&) = default;
+			TMapAttributeHelper& operator=(TMapAttributeHelper&&) = default;
+
+			~TMapAttributeHelper()
+			{
+				Attributes.Reset();
+			}
+
+			void SetKeyValue(const KeyType& InKey, const ValueType& InValue)
+			{
+				const uint32 Hash = GetTypeHash(InKey);
+				SetKeyValueByHash(Hash, InKey, InValue);
+			}
+
+			bool GetValue(const KeyType& InKey, ValueType& OutValue) const
+			{
+				const uint32 Hash = GetTypeHash(InKey);
+				return GetValueByHash(Hash, InKey, OutValue);
+			}
+
+			bool RemoveKey(const KeyType& InKey)
+			{
+				const uint32 Hash = GetTypeHash(InKey);
+				return RemoveKeyByHash(Hash, InKey);
+			}
+
+			bool RemoveKeyAndGetValue(const KeyType& InKey, ValueType& OutValue)
+			{
+				const uint32 Hash = GetTypeHash(InKey);
+				return RemoveKeyAndGetValueByHash(Hash, InKey, OutValue);
+			}
+
+			void SetKeyValueByHash(uint32 Hash, const KeyType& InKey, const ValueType& InValue)
+			{
+				TSharedPtr<FAttributeStorage, ESPMode::ThreadSafe> AttributesPtr = Attributes.Pin();
+				if (!ensure(AttributesPtr.IsValid()))
+				{
+					return;
+				}
+
+				if (TPair<TAttributeHandle<KeyType>, TAttributeHandle<ValueType>>* Pair = CachedKeysAndValues.FindByHash(Hash, InKey))
+				{
+					Pair->Value.Set(InValue);
+				}
+				else
+				{
+					FAttributeKey IndexKey = GetKeyAttribute(CachedKeysAndValues.Num());
+					FAttributeKey AttributeKey = GetValueAttribute(InKey);
+					check(AttributesPtr->RegisterAttribute<KeyType>(IndexKey, InKey) == EAttributeStorageResult::Operation_Success);
+					check(AttributesPtr->RegisterAttribute<ValueType>(AttributeKey, InValue) == EAttributeStorageResult::Operation_Success);
+					CachedKeysAndValues.AddByHash(
+						Hash
+						, InKey
+						, TPair<TAttributeHandle<KeyType>, TAttributeHandle<ValueType>>(
+							AttributesPtr->GetAttributeHandle<KeyType>(IndexKey)
+							, AttributesPtr->GetAttributeHandle<ValueType>(AttributeKey)
+						));
+					KeyCountHandle.Set(CachedKeysAndValues.Num());
+				}
+			}
+
+		
+			bool GetValueByHash(uint32 Hash, const KeyType& InKey, ValueType& OutValue) const
+			{
+				TSharedPtr<FAttributeStorage, ESPMode::ThreadSafe> AttributesPtr = Attributes.Pin();
+				if (!ensure(AttributesPtr.IsValid()))
+				{
+					return false;
+				}
+
+				if (const TPair<TAttributeHandle<KeyType>, TAttributeHandle<ValueType>>* Pair = CachedKeysAndValues.FindByHash(Hash, InKey))
+				{
+					return Pair->Value.Get(OutValue) == EAttributeStorageResult::Operation_Success;
+				}
+			}
+
+	
+			bool RemoveKeyByHash(uint32 Hash, const KeyType& InKey)
+			{
+				TSharedPtr<FAttributeStorage, ESPMode::ThreadSafe> AttributesPtr = Attributes.Pin();
+				if (!ensure(AttributesPtr.IsValid()))
+				{
+					return false;
+				}
+
+				if (TPair<TAttributeHandle<KeyType>, TAttributeHandle<ValueType>>* Pair = CachedKeysAndValues.FindByHash(Hash, InKey))
+				{
+					return RemoveBySwap(Hash, InKey, *Pair);
+				}
+
+				return false;
+			}
+
+			bool RemoveKeyAndGetValueByHash(uint32 Hash, const KeyType& InKey, ValueType& OutValue)
+			{
+				TSharedPtr<FAttributeStorage, ESPMode::ThreadSafe> AttributesPtr = Attributes.Pin();
+				if (!ensure(AttributesPtr.IsValid()))
+				{
+					return false;
+				}
+
+				if (TPair<TAttributeHandle<KeyType>, TAttributeHandle<ValueType>>* Pair = CachedKeysAndValues.FindByHash(Hash, InKey))
+				{
+					if (Pair->Value.Get(OutValue) == EAttributeStorageResult::Operation_Success)
+					{
+						return RemoveBySwap(Hash, InKey, *Pair);
+					}
+				}
+
+				return false;
+			}
+
+			void Reserve(int32 Number)
+			{
+				CachedKeysAndValues.Reserve(Number);
+			}
+
+			void Empty(int32 NumOfExpectedElements = 0)
+			{
+				EmptyInternal(NumOfExpectedElements);
+			}
+
+			TMapAttributeHelper& operator=(const TMap<KeyType, ValueType>& InMap)
+			{
+				TSharedPtr<FAttributeStorage, ESPMode::ThreadSafe> AttributesPtr = Attributes.Pin();
+				if (!ensure(AttributesPtr.IsValid()))
+				{
+					return *this;
+				}
+
+				// Empty
+				EmptyInternal(InMap.Num());
+
+				KeyCountHandle.Set(InMap.Num());
+
+				for (const TPair<KeyType, ValueType>& Pair : InMap)
+				{
+					SetKeyValue(Pair.Key, Pair.Value);
+				}
+
+				return *this;
+			}
+
+			TMap<KeyType, ValueType> ToMap() const
+			{
+				TMap<KeyType, ValueType> Map;
+				Map.Reserve(CachedKeysAndValues.Num());
+				
+				for (const TPair<KeyType, TPair<TAttributeHandle<KeyType>, TAttributeHandle<ValueType>>>& Pair: CachedKeysAndValues)
+				{
+					KeyType Key;
+					Pair.Value.Key.Get(Key);
+					ValueType Value;
+					Pair.Value.Value.Get(Value);
+
+					Map.Add(MoveTemp(Key), MoveTemp(Value));
+				}
+
+				return Map;
+			}
+
+		private:
+			FAttributeKey GetKeyAttribute(int32 Index) const
+			{
+				static const FString KeyIndex = TEXT("_KeyIndex_");
+				FString IndexedKey =  KeyCountHandle.GetKey().ToString();
+				IndexedKey.Reserve(KeyIndex.Len() + 16 /*Max size for a int32*/);
+				IndexedKey.Append(KeyIndex);
+				IndexedKey.AppendInt(Index);
+				return FAttributeKey(MoveTemp(IndexedKey));
+			}
+
+			FAttributeKey GetValueAttribute(const KeyType& InKey) const
+			{
+				static const FString KeyChars = TEXT("_Key_");
+				FString ValueAttribute =  KeyCountHandle.GetKey().ToString();
+				FString Key = TTypeToString<KeyType>::ToString(InKey);
+				ValueAttribute.Reserve(KeyChars.Len() + Key.Len());
+				ValueAttribute.Append(KeyChars);
+				ValueAttribute.Append(Key);
+				return FAttributeKey(MoveTemp(ValueAttribute));
+			}
+
+			TAttributeHandle<KeyType> GetLastKeyAttributeHandle() const
+			{
+				TSharedPtr<FAttributeStorage, ESPMode::ThreadSafe> AttributesPtr = Attributes.Pin();
+				if (!ensure(AttributesPtr.IsValid()))
+				{
+					return {};
+				}
+				return AttributesPtr->GetAttributeHandle<KeyType>(GetKeyAttribute(CachedKeysAndValues.Num() - 1));
+			}
+
+			bool RemoveBySwap(uint32 Hash, const KeyType& InKey, TPair<TAttributeHandle<FString>, TAttributeHandle<ValueType>>& CachedPair)
+			{
+				TSharedPtr<FAttributeStorage, ESPMode::ThreadSafe> AttributesPtr = Attributes.Pin();
+				if (!ensure(AttributesPtr.IsValid()))
+				{
+					return false;
+				}
+
+				TAttributeHandle<KeyType> LastKeyIndex = GetLastKeyAttributeHandle();
+				KeyType LastKey;
+				LastKeyIndex.Get(LastKey);
+				CachedKeysAndValues[LastKey].Key = CachedPair.Key;
+				CachedPair->Key.Set(MoveTemp(LastKey));
+
+				AttributesPtr->UnregisterAttribute(LastKeyIndex.GetKey());
+				AttributesPtr->UnregisterAttribute(CachedPair.Value.GetKey());
+				CachedKeysAndValues.RemoveByHash(Hash, InKey);
+				KeyCountHandle.Set(CachedKeysAndValues.Num());
+			}
+
+			void EmptyInternal(int32 NumOfExpectedElements)
+			{
+				TSharedPtr<FAttributeStorage, ESPMode::ThreadSafe> AttributesPtr = Attributes.Pin();
+				if (!ensure(AttributesPtr.IsValid()))
+				{
+					return;
+				}
+
+				for (const TPair<KeyType, TPair<TAttributeHandle<KeyType>, TAttributeHandle<ValueType>>>& Pair : CachedKeysAndValues)
+				{
+					const FAttributeKey& KeyAttribute = Pair.Value.Key.GetKey();
+					AttributesPtr->UnregisterAttribute(KeyAttribute);
+
+					const FAttributeKey& ValueAttribute = Pair.Value.Value.GetKey();
+					AttributesPtr->UnregisterAttribute(ValueAttribute);
+				}
+
+				CachedKeysAndValues.Empty(NumOfExpectedElements);
+			}
+
+			TMap<KeyType, TPair<TAttributeHandle<KeyType>, TAttributeHandle<ValueType>>> CachedKeysAndValues;
+			TAttributeHandle<int32> KeyCountHandle; //Assign in Initialize function, it will ensure if its the default value (NAME_None) when using the class
+			TWeakPtr<FAttributeStorage, ESPMode::ThreadSafe> Attributes = nullptr;
+		};
+
 		/**
 		 * Helper struct use to declare static const data we use in the UInterchangeBaseNode
 		 * Node that derive from UInterchangeBaseNode can also add a struct that derive from this one to add there static data
diff --git a/Engine/Source/Runtime/Interchange/Core/Public/Types/AttributeStorage.h b/Engine/Source/Runtime/Interchange/Core/Public/Types/AttributeStorage.h
--- a/Engine/Source/Runtime/Interchange/Core/Public/Types/AttributeStorage.h
+++ b/Engine/Source/Runtime/Interchange/Core/Public/Types/AttributeStorage.h
@@ -11,23 +11,21 @@
 #include "Misc/NetworkGuid.h"
 #include "Misc/ScopeLock.h"
 
-
 //Interchange namespace
 namespace UE
 {
 	namespace Interchange
 	{
-
 		struct FAttributeKey
 		{
 			FString Key;
 
 			FAttributeKey() = default;
 		
-			FAttributeKey(const FAttributeKey& Other)
-			{
-				Key = Other.Key;
-			}
+			FAttributeKey(const FAttributeKey& Other) = default;
+			FAttributeKey(FAttributeKey&& Other) = default;
+			FAttributeKey& operator=(const FAttributeKey&) = default;
+			FAttributeKey& operator=(FAttributeKey&&) = default;
 
 			explicit FAttributeKey(const FName& Other)
 			{
@@ -39,6 +37,11 @@ namespace UE
 				Key = Other;
 			}
 
+			explicit FAttributeKey(FString&& Other)
+			{
+				Key = MoveTemp(Other);
+			}
+
 			explicit FAttributeKey(const FText& Other)
 			{
 				Key = Other.ToString();
@@ -55,12 +58,6 @@ namespace UE
 				return *(Key);
 			}
 
-			FORCEINLINE FAttributeKey& operator=(const FAttributeKey& Other)
-			{
-				Key = Other.Key;
-				return *this;
-			}
-
 			FORCEINLINE FAttributeKey& operator=(const FName& Other)
 			{
 				Key = Other.ToString();
@@ -73,6 +70,12 @@ namespace UE
 				return *this;
 			}
 
+			FORCEINLINE FAttributeKey& operator=(FString&& Other)
+			{
+				Key = MoveTemp(Other);
+				return *this;
+			}
+
 			FORCEINLINE FAttributeKey& operator=(const FText& Other)
 			{
 				Key = Other.ToString();
@@ -128,11 +131,13 @@ namespace UE
 			}
 		};
 
+		using ::GetTypeHash;
 		FORCEINLINE uint32 GetTypeHash(const FAttributeKey& AttributeKey)
 		{
 			return GetTypeHash(AttributeKey.Key);
 		}
-	
+
+
 		/**
 		 * Enumerates the built-in types that can be stored in instances of FAttributeStorage.
 		 * We cannot change the value of a type to make sure the serialization of old asset is always working
@@ -374,6 +379,12 @@ namespace UE
 					}
 					return EAttributeStorageResult::Operation_Error_InvalidStorage;
 				}
+
+				const FAttributeKey& GetKey() const
+				{
+					return Key;
+				}
+
 			protected:
 				class FAttributeStorage* AttributeStorage;
 				FAttributeKey Key;
@@ -1299,4 +1310,4 @@ namespace UE
 		};
 
 	} //ns interchange
-}
\ No newline at end of file
+}
diff --git a/Engine/Source/Runtime/Interchange/Engine/Private/InterchangeManager.cpp b/Engine/Source/Runtime/Interchange/Engine/Private/InterchangeManager.cpp
--- a/Engine/Source/Runtime/Interchange/Engine/Private/InterchangeManager.cpp
+++ b/Engine/Source/Runtime/Interchange/Engine/Private/InterchangeManager.cpp
@@ -359,16 +359,20 @@ bool UInterchangeManager::RegisterTranslator(const UClass* TranslatorClass)
 		return false;
 	}
 
-	if(RegisteredTranslators.Contains(TranslatorClass))
+	uint32 TranslatorClassHash = GetTypeHash(TranslatorClass);
+	if(RegisteredTranslatorsClass.ContainsByHash(TranslatorClassHash, TranslatorClass))
 	{
 		return true;
 	}
+	RegisteredTranslatorsClass.AddByHash(TranslatorClassHash, TranslatorClass);
+
+
 	UInterchangeTranslatorBase* TranslatorToRegister = NewObject<UInterchangeTranslatorBase>(GetTransientPackage(), TranslatorClass, NAME_None);
 	if(!TranslatorToRegister)
 	{
 		return false;
 	}
-	RegisteredTranslators.Add(TranslatorClass, TranslatorToRegister);
+	RegisteredTranslators.Add(TranslatorToRegister);
 	return true;
 }
 
@@ -415,10 +419,8 @@ bool UInterchangeManager::RegisterWriter(const UClass* WriterClass)
 
 bool UInterchangeManager::CanTranslateSourceData(const UInterchangeSourceData* SourceData) const
 {
-	//Found the translator
-	UE::Interchange::FScopedTranslator ScopeDataTranslator(SourceData);
-	const UInterchangeTranslatorBase* SourceDataTranslator = ScopeDataTranslator.GetTranslator();
-	if (SourceDataTranslator)
+	// Can we find a translator
+	if (GetTranslatorForSourceData(SourceData))
 	{
 		return true;
 	}
@@ -555,7 +557,6 @@ UE::Interchange::FAssetImportResultRef UInterchangeManager::ImportAssetAsync(con
 	UInterchangeSourceData* DuplicateSourceData = Cast<UInterchangeSourceData>(StaticDuplicateObject(SourceData, GetTransientPackage()));
 	//Array of source data to build one graph per source
 	AsyncHelper->SourceDatas.Add(DuplicateSourceData);
-		
 
 	//Get all the translators for the source datas
 	for (int32 SourceDataIndex = 0; SourceDataIndex < AsyncHelper->SourceDatas.Num(); ++SourceDataIndex)
@@ -679,16 +680,12 @@ void UInterchangeManager::ReleaseAsyncHelper(TWeakPtr<UE::Interchange::FImportAs
 
 UInterchangeTranslatorBase* UInterchangeManager::GetTranslatorForSourceData(const UInterchangeSourceData* SourceData) const
 {
-	if (RegisteredTranslators.Num() == 0)
-	{
-		return nullptr;
-	}
-	//Found the translator
-	for (const auto Kvp : RegisteredTranslators)
+	// Find the translator
+	for (UInterchangeTranslatorBase* Translator : RegisteredTranslators)
 	{
-		if (Kvp.Value->CanImportSourceData(SourceData))
+		if (Translator->CanImportSourceData(SourceData))
 		{
-			UInterchangeTranslatorBase* SourceDataTranslator = NewObject<UInterchangeTranslatorBase>(GetTransientPackage(), Kvp.Key, NAME_None);
+			UInterchangeTranslatorBase* SourceDataTranslator = NewObject<UInterchangeTranslatorBase>(GetTransientPackage(), Translator->GetClass(), NAME_None);
 			return SourceDataTranslator;
 		}
 	}
@@ -712,6 +709,30 @@ bool UInterchangeManager::WarnIfInterchangeIsActive()
 	return true;
 }
 
+bool UInterchangeManager::CanTranslateSourceDataWithPayloadInterface(const UInterchangeSourceData* SourceData, const UClass* PayloadInterfaceClass) const
+{
+	if (GetTranslatorSupportingPayloadInterfaceForSourceData(SourceData, PayloadInterfaceClass))
+	{
+		return true;
+	}
+	return false;
+}
+
+UInterchangeTranslatorBase* UInterchangeManager::GetTranslatorSupportingPayloadInterfaceForSourceData(const UInterchangeSourceData* SourceData, const UClass* PayloadInterfaceClass) const
+{
+	// Find the translator
+	for (UInterchangeTranslatorBase* Translator : RegisteredTranslators)
+	{
+		const UClass* TranslatorClass = Translator->GetClass();
+		if (TranslatorClass->ImplementsInterface(PayloadInterfaceClass) && Translator->CanImportSourceData(SourceData))
+		{
+			UInterchangeTranslatorBase* SourceDataTranslator = NewObject<UInterchangeTranslatorBase>(GetTransientPackage(), TranslatorClass, NAME_None);
+			return SourceDataTranslator;
+		}
+	}
+	return nullptr;
+}
+
 bool UInterchangeManager::IsAttended()
 {
 	if (FApp::IsGame())
diff --git a/Engine/Source/Runtime/Interchange/Engine/Private/Tasks/InterchangeTaskCreateAsset.cpp b/Engine/Source/Runtime/Interchange/Engine/Private/Tasks/InterchangeTaskCreateAsset.cpp
--- a/Engine/Source/Runtime/Interchange/Engine/Private/Tasks/InterchangeTaskCreateAsset.cpp
+++ b/Engine/Source/Runtime/Interchange/Engine/Private/Tasks/InterchangeTaskCreateAsset.cpp
@@ -31,15 +31,9 @@ namespace UE
 				FString NodeDisplayName = Node->GetDisplayLabel();
 				const FString BaseFileName = FPaths::GetBaseFilename(SourceData->GetFilename());
 
-				//Set the asset name and the package name
-				//if (NodeDisplayName.Equals(BaseFileName) || BaseFileName.IsEmpty())
-				{
-					OutAssetName = NodeDisplayName;
-				}
-// 				else
-// 				{
-// 					OutAssetName = BaseFileName.IsEmpty() ? NodeDisplayName : BaseFileName + TEXT("_") + NodeDisplayName;
-// 				}
+				// Set the asset name and the package name
+				OutAssetName = NodeDisplayName;
+
 				OutPackageName = FPaths::Combine(*PackageBasePath, *OutAssetName);
 
 				//Sanitize only the package name
diff --git a/Engine/Source/Runtime/Interchange/Engine/Public/InterchangeManager.h b/Engine/Source/Runtime/Interchange/Engine/Public/InterchangeManager.h
--- a/Engine/Source/Runtime/Interchange/Engine/Public/InterchangeManager.h
+++ b/Engine/Source/Runtime/Interchange/Engine/Public/InterchangeManager.h
@@ -246,6 +246,7 @@ public:
 	 * @return true if the translator class can be register false otherwise.
 	 *
 	 * @Note if you register multiple time the same class it will return true for every call
+	 * @Note The order in which the translators are registered will be the same as the order used to select a translator to import a file
 	 */
 	bool RegisterTranslator(const UClass* TranslatorClass);
 
@@ -369,6 +370,21 @@ public:
 	 */
 	bool WarnIfInterchangeIsActive();
 
+	/**
+	 * Look if there is a translator registered that can translate the source data with the specified PayloadInterface
+	 * @Param SourceData - The source data input we want to translate to Uod
+	 * @return true if the source data can be translated using the specified PayloadInterface, false otherwise.
+	 */
+	bool CanTranslateSourceDataWithPayloadInterface(const UInterchangeSourceData* SourceData, const UClass* PayloadInterfaceClass) const;
+
+	/*
+	 * Return the first translator that can translate the source data with the specified PayloadInterface.
+	 * @Param SourceData - The source data for which we search a translator.
+	 * @Param PayloadInterfaceClass - The PayloadInterface that the translator must implement.
+	 * @return return a matching translator implementing the specified PayloadInterface or nullptr if there is no match.
+	 */
+	UInterchangeTranslatorBase* GetTranslatorSupportingPayloadInterfaceForSourceData(const UInterchangeSourceData* SourceData, const UClass* PayloadInterfaceClass) const;
+
 protected:
 
 	/** Return true if we can show some UI */
@@ -422,10 +438,14 @@ private:
 	TSharedPtr<FAsyncTaskNotification> Notification = nullptr;
 	FDelegateHandle NotificationTickHandle;
 
+	// Caching the registered translator classes to avoid double registration fast
+	UPROPERTY()
+	TSet<const UClass*> RegisteredTranslatorsClass;
+
 	//The manager will create translator at every import, translator must be able to retrieve payload information when the factory ask for it.
 	//The translator stored has value is only use to know if we can use this type of translator.
 	UPROPERTY()
-	TMap<const UClass*, UInterchangeTranslatorBase* > RegisteredTranslators;
+	TArray<UInterchangeTranslatorBase*> RegisteredTranslators;
 	
 	//The manager will create only one pipeline per type
 	UPROPERTY()
diff --git a/Engine/Source/Runtime/TextureUtilitiesCommon/Private/TextureUtilitiesCommonModule.cpp b/Engine/Source/Runtime/TextureUtilitiesCommon/Private/TextureUtilitiesCommonModule.cpp
new file mode 100644
--- /dev/null
+++ b/Engine/Source/Runtime/TextureUtilitiesCommon/Private/TextureUtilitiesCommonModule.cpp
@@ -0,0 +1,10 @@
+// Copyright Epic Games, Inc. All Rights Reserved.
+
+#include "Modules/ModuleManager.h"
+
+/**
+ * A simple module that act as bridge between the Interchange plugin and the UnrealEd module
+ * The goal of this module is to simply avoid some code duplication between the Interchange system and the old editor factories for the textures
+ */
+
+IMPLEMENT_MODULE(FDefaultModuleImpl, TextureUtilitiesCommon)
diff --git a/Engine/Source/Runtime/TextureUtilitiesCommon/Private/UDIMUtilities.cpp b/Engine/Source/Runtime/TextureUtilitiesCommon/Private/UDIMUtilities.cpp
new file mode 100644
--- /dev/null
+++ b/Engine/Source/Runtime/TextureUtilitiesCommon/Private/UDIMUtilities.cpp
@@ -0,0 +1,50 @@
+// Copyright Epic Games, Inc. All Rights Reserved.
+
+#include "UDIMUtilities.h"
+
+#include "Internationalization/Regex.h"
+
+namespace UE
+{
+	namespace TextureUtilitiesCommon
+	{
+		uint32 ParseUDIMName(const FString& Name, const FString& UdimRegexPattern, FString& OutPrefixName, FString& OutPostfixName)
+		{
+			FRegexPattern RegexPattern( UdimRegexPattern );
+			FRegexMatcher RegexMatcher( RegexPattern, Name );
+
+			int32 UdimValue = INDEX_NONE;
+
+			if ( RegexMatcher.FindNext() )
+			{
+				const int32 StartOfCaptureGroup1 = RegexMatcher.GetCaptureGroupBeginning(1);
+				const int32 EndOfCaptureGroup1 = RegexMatcher.GetCaptureGroupEnding(1);
+				const int32 StartOfCaptureGroup2 = RegexMatcher.GetCaptureGroupBeginning(2);
+				const int32 EndOfCaptureGroup2 = RegexMatcher.GetCaptureGroupEnding(2);
+				const int32 StartOfCaptureGroup3 = RegexMatcher.GetCaptureGroupBeginning(3);
+				const int32 EndOfCaptureGroup3 = RegexMatcher.GetCaptureGroupEnding(3);
+
+				if ( StartOfCaptureGroup1 != INDEX_NONE && StartOfCaptureGroup2 != INDEX_NONE &&
+					 EndOfCaptureGroup1 != INDEX_NONE && EndOfCaptureGroup2 != INDEX_NONE )
+				{
+					LexFromString( UdimValue, *Name.Mid( StartOfCaptureGroup2, EndOfCaptureGroup2 - StartOfCaptureGroup2 ) );
+
+					OutPrefixName = Name.Mid( StartOfCaptureGroup1, EndOfCaptureGroup1 - StartOfCaptureGroup1 );
+
+					if ( StartOfCaptureGroup3 != INDEX_NONE && EndOfCaptureGroup3 != INDEX_NONE )
+					{
+						OutPostfixName = Name.Mid( StartOfCaptureGroup3, EndOfCaptureGroup3 - StartOfCaptureGroup3 );
+					}
+				}
+			}
+	
+			if ( UdimValue < 1001 )
+			{
+				// UDIM starts with 1001 as the origin
+				return INDEX_NONE;
+			}
+
+			return UdimValue;
+		}
+	}
+}
diff --git a/Engine/Source/Runtime/TextureUtilitiesCommon/Public/UDIMUtilities.h b/Engine/Source/Runtime/TextureUtilitiesCommon/Public/UDIMUtilities.h
new file mode 100644
--- /dev/null
+++ b/Engine/Source/Runtime/TextureUtilitiesCommon/Public/UDIMUtilities.h
@@ -0,0 +1,15 @@
+// Copyright Epic Games, Inc. All Rights Reserved.
+
+#pragma once
+
+#include "CoreMinimal.h"
+
+class FString;
+
+namespace UE
+{
+	namespace TextureUtilitiesCommon
+	{
+		TEXTUREUTILITIESCOMMON_API uint32 ParseUDIMName(const FString& Name, const FString& UdimRegexPattern, FString& OutPrefixName, FString& OutPostfixName);
+	}
+}
diff --git a/Engine/Source/Runtime/TextureUtilitiesCommon/TextureUtilitiesCommon.build.cs b/Engine/Source/Runtime/TextureUtilitiesCommon/TextureUtilitiesCommon.build.cs
new file mode 100644
--- /dev/null
+++ b/Engine/Source/Runtime/TextureUtilitiesCommon/TextureUtilitiesCommon.build.cs
@@ -0,0 +1,11 @@
+// Copyright Epic Games, Inc. All Rights Reserved.
+
+using UnrealBuildTool;
+
+public class TextureUtilitiesCommon : ModuleRules
+{
+	public TextureUtilitiesCommon(ReadOnlyTargetRules Target) : base(Target)
+	{
+		PublicDependencyModuleNames.Add("Core");
+	}
+}
\ No newline at end of file
