Interchange Framework: cleanup module naming

diff --git a/Engine/Plugins/Editor/InterchangeEditor/Source/InterchangeEditor/InterchangeEditor.Build.cs b/Engine/Plugins/Editor/InterchangeEditor/Source/InterchangeEditor/InterchangeEditor.Build.cs
--- a/Engine/Plugins/Editor/InterchangeEditor/Source/InterchangeEditor/InterchangeEditor.Build.cs
+++ b/Engine/Plugins/Editor/InterchangeEditor/Source/InterchangeEditor/InterchangeEditor.Build.cs
@@ -20,9 +20,9 @@ namespace UnrealBuildTool.Rules
 				new string[]
 				{
 					"InterchangeCore",
-					"InterchangeImportPlugin",
-					"InterchangeNodePlugin",
-					"InterchangeExportPlugin",
+					"InterchangeImport",
+					"InterchangeNodes",
+					"InterchangeExport",
 					"InterchangeDispatcher",
 					"InterchangeFbxParser"
 				}
diff --git a/Engine/Plugins/Enterprise/DatasmithImporter/Source/DatasmithImporter/DatasmithImporter.Build.cs b/Engine/Plugins/Enterprise/DatasmithImporter/Source/DatasmithImporter/DatasmithImporter.Build.cs
--- a/Engine/Plugins/Enterprise/DatasmithImporter/Source/DatasmithImporter/DatasmithImporter.Build.cs
+++ b/Engine/Plugins/Enterprise/DatasmithImporter/Source/DatasmithImporter/DatasmithImporter.Build.cs
@@ -27,7 +27,7 @@ namespace UnrealBuildTool.Rules
 					"InputCore",
 					"InterchangeCore",
 					"InterchangeEngine",
-					"InterchangeNodePlugin",
+					"InterchangeNodes",
 					"Json",
 					"Landscape",
 					"LandscapeEditor",
diff --git a/Engine/Plugins/Runtime/Interchange/Interchange.uplugin b/Engine/Plugins/Runtime/Interchange/Interchange.uplugin
--- a/Engine/Plugins/Runtime/Interchange/Interchange.uplugin
+++ b/Engine/Plugins/Runtime/Interchange/Interchange.uplugin
@@ -16,22 +16,22 @@
 	"Installed": false,
 	"Modules": [
 		{
-			"Name": "InterchangeNodePlugin",
+			"Name": "InterchangeNodes",
 			"Type": "Runtime",
 			"LoadingPhase": "Default"
 		},
 		{
-			"Name": "InterchangeImportPlugin",
+			"Name": "InterchangeImport",
 			"Type": "Runtime",
 			"LoadingPhase": "Default"
 		},
 		{
-			"Name": "InterchangeExportPlugin",
+			"Name": "InterchangeExport",
 			"Type": "Runtime",
 			"LoadingPhase": "Default"
 		},
 		{
-			"Name": "InterchangePipelinePlugin",
+			"Name": "InterchangePipelines",
 			"Type": "Runtime",
 			"LoadingPhase": "Default"
 		},
diff --git a/Engine/Plugins/Runtime/Interchange/Source/Export/InterchangeExport.Build.cs b/Engine/Plugins/Runtime/Interchange/Source/Export/InterchangeExport.Build.cs
new file mode 100644
--- /dev/null
+++ b/Engine/Plugins/Runtime/Interchange/Source/Export/InterchangeExport.Build.cs
@@ -0,0 +1,21 @@
+// Copyright Epic Games, Inc. All Rights Reserved.
+
+namespace UnrealBuildTool.Rules
+{
+	public class InterchangeExport : ModuleRules
+	{
+		public InterchangeExport(ReadOnlyTargetRules Target) : base(Target)
+		{
+			PublicDependencyModuleNames.AddRange(
+				new string[]
+				{
+					"Core",
+					"CoreUObject",
+					"Engine",
+					"InterchangeCore",
+					"InterchangeEngine"
+				}
+			);
+		}
+	}
+}
diff --git a/Engine/Plugins/Runtime/Interchange/Source/Export/InterchangeExportPlugin.Build.cs b/Engine/Plugins/Runtime/Interchange/Source/Export/InterchangeExportPlugin.Build.cs
deleted file mode 100644
diff --git a/Engine/Plugins/Runtime/Interchange/Source/Export/Private/InterchangeExportModule.cpp b/Engine/Plugins/Runtime/Interchange/Source/Export/Private/InterchangeExportModule.cpp
new file mode 100644
--- /dev/null
+++ b/Engine/Plugins/Runtime/Interchange/Source/Export/Private/InterchangeExportModule.cpp
@@ -0,0 +1,50 @@
+// Copyright Epic Games, Inc. All Rights Reserved.
+
+#include "InterchangeExportModule.h"
+
+#include "CoreMinimal.h"
+#include "Engine/Engine.h"
+#include "InterchangeManager.h"
+#include "InterchangeTextureWriter.h"
+#include "Misc/CoreDelegates.h"
+#include "Modules/ModuleManager.h"
+
+
+class FInterchangeExportModule : public IInterchangeExportModule
+{
+	/** IModuleInterface implementation */
+	virtual void StartupModule() override;
+	virtual void ShutdownModule() override;
+};
+
+IMPLEMENT_MODULE(FInterchangeExportModule, InterchangeExport)
+
+
+
+void FInterchangeExportModule::StartupModule()
+{
+	auto RegisterItems = []()
+	{
+		UInterchangeManager& InterchangeManager = UInterchangeManager::GetInterchangeManager();
+		//Register the Writers
+		InterchangeManager.RegisterWriter(UInterchangeTextureWriter::StaticClass());
+	};
+
+	if (GEngine)
+	{
+		RegisterItems();
+	}
+	else
+	{
+		FCoreDelegates::OnPostEngineInit.AddLambda(RegisterItems);
+	}
+}
+
+
+void FInterchangeExportModule::ShutdownModule()
+{
+	
+}
+
+
+
diff --git a/Engine/Plugins/Runtime/Interchange/Source/Export/Private/InterchangeExportPlugin.cpp b/Engine/Plugins/Runtime/Interchange/Source/Export/Private/InterchangeExportPlugin.cpp
deleted file mode 100644
diff --git a/Engine/Plugins/Runtime/Interchange/Source/Export/Public/IInterchangeExportPlugin.h b/Engine/Plugins/Runtime/Interchange/Source/Export/Public/IInterchangeExportPlugin.h
deleted file mode 100644
diff --git a/Engine/Plugins/Runtime/Interchange/Source/Export/Public/InterchangeExportModule.h b/Engine/Plugins/Runtime/Interchange/Source/Export/Public/InterchangeExportModule.h
new file mode 100644
--- /dev/null
+++ b/Engine/Plugins/Runtime/Interchange/Source/Export/Public/InterchangeExportModule.h
@@ -0,0 +1,39 @@
+// Copyright Epic Games, Inc. All Rights Reserved.
+
+#pragma once
+
+#include "CoreMinimal.h"
+#include "Modules/ModuleInterface.h"
+#include "Modules/ModuleManager.h"
+
+
+/**
+ * The public interface to this module
+ */
+class IInterchangeExportModule : public IModuleInterface
+{
+
+public:
+
+	/**
+	 * Singleton-like access to this module's interface.  This is just for convenience!
+	 * Beware of calling this during the shutdown phase, though.  Your module might have been unloaded already.
+	 *
+	 * @return Returns singleton instance, loading the module on demand if needed
+	 */
+	static inline IInterchangeExportModule& Get()
+	{
+		return FModuleManager::LoadModuleChecked< IInterchangeExportModule >( "InterchangeExport" );
+	}
+
+	/**
+	 * Checks to see if this module is loaded and ready.  It is only valid to call Get() if IsAvailable() returns true.
+	 *
+	 * @return True if the module is loaded and ready to use
+	 */
+	static inline bool IsAvailable()
+	{
+		return FModuleManager::Get().IsModuleLoaded( "InterchangeExport" );
+	}
+};
+
diff --git a/Engine/Plugins/Runtime/Interchange/Source/Export/Public/InterchangeTextureWriter.h b/Engine/Plugins/Runtime/Interchange/Source/Export/Public/InterchangeTextureWriter.h
--- a/Engine/Plugins/Runtime/Interchange/Source/Export/Public/InterchangeTextureWriter.h
+++ b/Engine/Plugins/Runtime/Interchange/Source/Export/Public/InterchangeTextureWriter.h
@@ -11,7 +11,7 @@
 #include "InterchangeTextureWriter.generated.h"
 
 UCLASS(BlueprintType)
-class INTERCHANGEEXPORTPLUGIN_API UInterchangeTextureWriter : public UInterchangeWriterBase
+class INTERCHANGEEXPORT_API UInterchangeTextureWriter : public UInterchangeWriterBase
 {
 	GENERATED_BODY()
 public:
diff --git a/Engine/Plugins/Runtime/Interchange/Source/Import/InterchangeImport.Build.cs b/Engine/Plugins/Runtime/Interchange/Source/Import/InterchangeImport.Build.cs
new file mode 100644
--- /dev/null
+++ b/Engine/Plugins/Runtime/Interchange/Source/Import/InterchangeImport.Build.cs
@@ -0,0 +1,40 @@
+// Copyright Epic Games, Inc. All Rights Reserved.
+
+namespace UnrealBuildTool.Rules
+{
+	public class InterchangeImport : ModuleRules
+	{
+		public InterchangeImport(ReadOnlyTargetRules Target) : base(Target)
+		{
+			PublicDependencyModuleNames.AddRange(
+				new string[]
+				{
+					"Core",
+					"CoreUObject",
+					"Engine",
+					"InterchangeCore",
+					"InterchangeDispatcher",
+					"InterchangeEngine",
+					"InterchangeNodes",
+					"MeshDescription",
+					"StaticMeshDescription",
+					"SkeletalMeshDescription",
+				}
+			);
+			PrivateDependencyModuleNames.AddRange(
+				new string[]
+				{
+					"ImageWrapper",
+					"InterchangeDispatcher",
+					"Json",
+					"RHI",
+					"TextureUtilitiesCommon",
+				}
+			);
+
+			OptimizeCode = CodeOptimization.Never;
+			bUseUnity = false;
+			PCHUsage = PCHUsageMode.NoPCHs;
+		}
+	}
+}
diff --git a/Engine/Plugins/Runtime/Interchange/Source/Import/InterchangeImportPlugin.Build.cs b/Engine/Plugins/Runtime/Interchange/Source/Import/InterchangeImportPlugin.Build.cs
deleted file mode 100644
diff --git a/Engine/Plugins/Runtime/Interchange/Source/Import/Private/Fbx/InterchangeFbxTranslator.cpp b/Engine/Plugins/Runtime/Interchange/Source/Import/Private/Fbx/InterchangeFbxTranslator.cpp
--- a/Engine/Plugins/Runtime/Interchange/Source/Import/Private/Fbx/InterchangeFbxTranslator.cpp
+++ b/Engine/Plugins/Runtime/Interchange/Source/Import/Private/Fbx/InterchangeFbxTranslator.cpp
@@ -7,7 +7,7 @@
 #include "InterchangeDispatcher.h"
 #include "InterchangeDispatcherTask.h"
 #include "InterchangeManager.h"
-#include "LogInterchangeImportPlugin.h"
+#include "InterchangeImportLog.h"
 #include "Mesh/InterchangeSkeletalMeshPayload.h"
 #include "Mesh/InterchangeStaticMeshPayload.h"
 #include "Misc/ConfigCacheIni.h"
diff --git a/Engine/Plugins/Runtime/Interchange/Source/Import/Private/InterchangeImportLog.h b/Engine/Plugins/Runtime/Interchange/Source/Import/Private/InterchangeImportLog.h
new file mode 100644
--- /dev/null
+++ b/Engine/Plugins/Runtime/Interchange/Source/Import/Private/InterchangeImportLog.h
@@ -0,0 +1,7 @@
+// Copyright Epic Games, Inc. All Rights Reserved.
+
+#pragma once
+
+#include "CoreMinimal.h"
+
+DECLARE_LOG_CATEGORY_EXTERN(LogInterchangeImport, Log, All);
diff --git a/Engine/Plugins/Runtime/Interchange/Source/Import/Private/InterchangeImportModule.cpp b/Engine/Plugins/Runtime/Interchange/Source/Import/Private/InterchangeImportModule.cpp
new file mode 100644
--- /dev/null
+++ b/Engine/Plugins/Runtime/Interchange/Source/Import/Private/InterchangeImportModule.cpp
@@ -0,0 +1,87 @@
+// Copyright Epic Games, Inc. All Rights Reserved.
+
+#include "InterchangeImportModule.h"
+
+#include "CoreMinimal.h"
+#include "Fbx/InterchangeFbxTranslator.h"
+#include "InterchangeImportLog.h"
+#include "InterchangeManager.h"
+#include "Material/InterchangeMaterialFactory.h"
+#include "Mesh/InterchangeSkeletalMeshFactory.h"
+#include "Mesh/InterchangeSkeletonFactory.h"
+#include "Misc/CoreDelegates.h"
+#include "Modules/ModuleManager.h"
+#include "Texture/InterchangeBMPTranslator.h"
+#include "Texture/InterchangeDDSTranslator.h"
+#include "Texture/InterchangeEXRTranslator.h"
+#include "Texture/InterchangeJPGTranslator.h"
+#include "Texture/InterchangePCXTranslator.h"
+#include "Texture/InterchangePNGTranslator.h"
+#include "Texture/InterchangePSDTranslator.h"
+#include "Texture/InterchangeTextureFactory.h"
+#include "Texture/InterchangeTGATranslator.h"
+#include "Texture/InterchangeUDIMTranslator.h"
+
+DEFINE_LOG_CATEGORY(LogInterchangeImport);
+
+class FInterchangeImportModule : public IInterchangeImportModule
+{
+	/** IModuleInterface implementation */
+	virtual void StartupModule() override;
+	virtual void ShutdownModule() override;
+};
+
+IMPLEMENT_MODULE(FInterchangeImportModule, InterchangeImport)
+
+
+
+void FInterchangeImportModule::StartupModule()
+{
+	auto RegisterItems = []()
+	{
+		UInterchangeManager& InterchangeManager = UInterchangeManager::GetInterchangeManager();
+
+		//Register the translators
+		//Scenes
+		InterchangeManager.RegisterTranslator(UInterchangeFbxTranslator::StaticClass()); //Do not submit uncommented until we replace completly fbx importer (staticmesh + skeletalMesh + animation)
+
+		//Textures
+
+		// UDIM must be registered before the other texture translators
+		InterchangeManager.RegisterTranslator(UInterchangeUDIMTranslator::StaticClass());
+
+
+		InterchangeManager.RegisterTranslator(UInterchangeBMPTranslator::StaticClass());
+		InterchangeManager.RegisterTranslator(UInterchangeDDSTranslator::StaticClass());
+		InterchangeManager.RegisterTranslator(UInterchangeEXRTranslator::StaticClass());
+		InterchangeManager.RegisterTranslator(UInterchangeJPGTranslator::StaticClass());
+		InterchangeManager.RegisterTranslator(UInterchangePCXTranslator::StaticClass());
+		InterchangeManager.RegisterTranslator(UInterchangePNGTranslator::StaticClass());
+		InterchangeManager.RegisterTranslator(UInterchangePSDTranslator::StaticClass());
+		InterchangeManager.RegisterTranslator(UInterchangeTGATranslator::StaticClass());
+
+		//Register the factories
+		InterchangeManager.RegisterFactory(UInterchangeTextureFactory::StaticClass());
+		InterchangeManager.RegisterFactory(UInterchangeMaterialFactory::StaticClass());
+		InterchangeManager.RegisterFactory(UInterchangeSkeletonFactory::StaticClass());
+		InterchangeManager.RegisterFactory(UInterchangeSkeletalMeshFactory::StaticClass());
+	};
+
+	if (GEngine)
+	{
+		RegisterItems();
+	}
+	else
+	{
+		FCoreDelegates::OnPostEngineInit.AddLambda(RegisterItems);
+	}
+}
+
+
+void FInterchangeImportModule::ShutdownModule()
+{
+
+}
+
+
+
diff --git a/Engine/Plugins/Runtime/Interchange/Source/Import/Private/InterchangeImportPlugin.cpp b/Engine/Plugins/Runtime/Interchange/Source/Import/Private/InterchangeImportPlugin.cpp
deleted file mode 100644
diff --git a/Engine/Plugins/Runtime/Interchange/Source/Import/Private/LogInterchangeImportPlugin.h b/Engine/Plugins/Runtime/Interchange/Source/Import/Private/LogInterchangeImportPlugin.h
deleted file mode 100644
diff --git a/Engine/Plugins/Runtime/Interchange/Source/Import/Private/Material/InterchangeMaterialFactory.cpp b/Engine/Plugins/Runtime/Interchange/Source/Import/Private/Material/InterchangeMaterialFactory.cpp
--- a/Engine/Plugins/Runtime/Interchange/Source/Import/Private/Material/InterchangeMaterialFactory.cpp
+++ b/Engine/Plugins/Runtime/Interchange/Source/Import/Private/Material/InterchangeMaterialFactory.cpp
@@ -2,9 +2,9 @@
 #include "Material/InterchangeMaterialFactory.h"
 
 #include "InterchangeImportCommon.h"
+#include "InterchangeImportLog.h"
 #include "InterchangeMaterialNode.h"
 #include "InterchangeSourceData.h"
-#include "LogInterchangeImportPlugin.h"
 #include "Materials/Material.h"
 #include "Materials/MaterialExpressionScalarParameter.h"
 #include "Materials/MaterialExpressionTextureCoordinate.h"
@@ -65,7 +65,7 @@ UObject* UInterchangeMaterialFactory::CreateEmptyAsset(const FCreateAssetParams&
 
 	if (!Material)
 	{
-		UE_LOG(LogInterchangeImportPlugin, Warning, TEXT("Could not create Material asset %s"), *Arguments.AssetName);
+		UE_LOG(LogInterchangeImport, Warning, TEXT("Could not create Material asset %s"), *Arguments.AssetName);
 		return nullptr;
 	}
 	Material->PreEditChange(nullptr);
@@ -77,7 +77,7 @@ UObject* UInterchangeMaterialFactory::CreateAsset(const UInterchangeMaterialFact
 {
 #if !WITH_EDITORONLY_DATA
 
-	UE_LOG(LogInterchangeImportPlugin, Error, TEXT("Cannot import Material asset in runtime, this is an editor only feature."));
+	UE_LOG(LogInterchangeImport, Error, TEXT("Cannot import Material asset in runtime, this is an editor only feature."));
 	return nullptr;
 
 #else
@@ -116,7 +116,7 @@ UObject* UInterchangeMaterialFactory::CreateAsset(const UInterchangeMaterialFact
 
 	if (!MaterialObject)
 	{
-		UE_LOG(LogInterchangeImportPlugin, Warning, TEXT("Could not create Material asset %s"), *Arguments.AssetName);
+		UE_LOG(LogInterchangeImport, Warning, TEXT("Could not create Material asset %s"), *Arguments.AssetName);
 		return nullptr;
 	}
 
diff --git a/Engine/Plugins/Runtime/Interchange/Source/Import/Private/Mesh/InterchangeSkeletalMeshFactory.cpp b/Engine/Plugins/Runtime/Interchange/Source/Import/Private/Mesh/InterchangeSkeletalMeshFactory.cpp
--- a/Engine/Plugins/Runtime/Interchange/Source/Import/Private/Mesh/InterchangeSkeletalMeshFactory.cpp
+++ b/Engine/Plugins/Runtime/Interchange/Source/Import/Private/Mesh/InterchangeSkeletalMeshFactory.cpp
@@ -5,6 +5,7 @@
 #include "Engine/SkeletalMesh.h"
 #include "GPUSkinPublicDefs.h"
 #include "InterchangeImportCommon.h"
+#include "InterchangeImportLog.h"
 #include "InterchangeJointNode.h"
 #include "InterchangeMaterialNode.h"
 #include "InterchangeSkeletalMeshLodDataNode.h"
@@ -12,7 +13,6 @@
 #include "InterchangeSkeletonNode.h"
 #include "InterchangeSourceData.h"
 #include "InterchangeTranslatorBase.h"
-#include "LogInterchangeImportPlugin.h"
 #include "Mesh/InterchangeSkeletalMeshPayload.h"
 #include "Mesh/InterchangeSkeletalMeshPayloadInterface.h"
 #include "Nodes/InterchangeBaseNode.h"
@@ -47,7 +47,7 @@ namespace UE
 				const UInterchangeJointNode* JointNode = Cast<UInterchangeJointNode>(NodeContainer->GetNode(JointNodeId));
 				if (!JointNode)
 				{
-					UE_LOG(LogInterchangeImportPlugin, Warning, TEXT("Invalid Skeleton Joint"));
+					UE_LOG(LogInterchangeImport, Warning, TEXT("Invalid Skeleton Joint"));
 					return;
 				}
 
@@ -101,7 +101,7 @@ namespace UE
 					const FTransform BoneTransform(BinaryBone.LocalTransform);
 					if (RefSkeleton.FindRawBoneIndex(BoneInfo.Name) != INDEX_NONE)
 					{
-						UE_LOG(LogInterchangeImportPlugin, Error, TEXT("Invalid Skeleton because of non-unique bone names [%s]"), *BoneInfo.Name.ToString());
+						UE_LOG(LogInterchangeImport, Error, TEXT("Invalid Skeleton because of non-unique bone names [%s]"), *BoneInfo.Name.ToString());
 						return false;
 					}
 					RefSkelModifier.Add(BoneInfo, BoneTransform);
@@ -998,7 +998,7 @@ UObject* UInterchangeSkeletalMeshFactory::CreateEmptyAsset(const FCreateAssetPar
 {
 #if !WITH_EDITOR || !WITH_EDITORONLY_DATA
 
-	UE_LOG(LogInterchangeImportPlugin, Error, TEXT("Cannot import skeletalMesh asset in runtime, this is an editor only feature."));
+	UE_LOG(LogInterchangeImport, Error, TEXT("Cannot import skeletalMesh asset in runtime, this is an editor only feature."));
 	return nullptr;
 
 #else
@@ -1030,7 +1030,7 @@ UObject* UInterchangeSkeletalMeshFactory::CreateEmptyAsset(const FCreateAssetPar
 	
 	if (!SkeletalMesh)
 	{
-		UE_LOG(LogInterchangeImportPlugin, Warning, TEXT("Could not create SkeletalMesh asset %s"), *Arguments.AssetName);
+		UE_LOG(LogInterchangeImport, Warning, TEXT("Could not create SkeletalMesh asset %s"), *Arguments.AssetName);
 		return nullptr;
 	}
 	
@@ -1046,7 +1046,7 @@ UObject* UInterchangeSkeletalMeshFactory::CreateAsset(const UInterchangeSkeletal
 {
 #if !WITH_EDITOR || !WITH_EDITORONLY_DATA
 
-	UE_LOG(LogInterchangeImportPlugin, Error, TEXT("Cannot import skeletalMesh asset in runtime, this is an editor only feature."));
+	UE_LOG(LogInterchangeImport, Error, TEXT("Cannot import skeletalMesh asset in runtime, this is an editor only feature."));
 	return nullptr;
 
 #else
@@ -1065,7 +1065,7 @@ UObject* UInterchangeSkeletalMeshFactory::CreateAsset(const UInterchangeSkeletal
 	const IInterchangeSkeletalMeshPayloadInterface* SkeletalMeshTranslatorPayloadInterface = Cast<IInterchangeSkeletalMeshPayloadInterface>(Arguments.Translator);
 	if (!SkeletalMeshTranslatorPayloadInterface)
 	{
-		UE_LOG(LogInterchangeImportPlugin, Error, TEXT("Cannot import skeletalMesh, the translator do not implement the IInterchangeSkeletalMeshPayloadInterface."));
+		UE_LOG(LogInterchangeImport, Error, TEXT("Cannot import skeletalMesh, the translator do not implement the IInterchangeSkeletalMeshPayloadInterface."));
 		return nullptr;
 	}
 
@@ -1092,7 +1092,7 @@ UObject* UInterchangeSkeletalMeshFactory::CreateAsset(const UInterchangeSkeletal
 
 	if (!SkeletalMeshObject)
 	{
-		UE_LOG(LogInterchangeImportPlugin, Error, TEXT("Could not create SkeletalMesh asset %s"), *Arguments.AssetName);
+		UE_LOG(LogInterchangeImport, Error, TEXT("Could not create SkeletalMesh asset %s"), *Arguments.AssetName);
 		return nullptr;
 	}
 
@@ -1105,7 +1105,7 @@ UObject* UInterchangeSkeletalMeshFactory::CreateAsset(const UInterchangeSkeletal
 			USkeletalMesh* SkeletalMesh = Cast<USkeletalMesh>(SkeletalMeshObject);
 			if (!ensure(SkeletalMesh))
 			{
-				UE_LOG(LogInterchangeImportPlugin, Error, TEXT("Could not create SkeletalMesh asset %s"), *Arguments.AssetName);
+				UE_LOG(LogInterchangeImport, Error, TEXT("Could not create SkeletalMesh asset %s"), *Arguments.AssetName);
 				return nullptr;
 			}
 			//Dirty the DDC Key for any imported Skeletal Mesh
@@ -1128,19 +1128,19 @@ UObject* UInterchangeSkeletalMeshFactory::CreateAsset(const UInterchangeSkeletal
 				const UInterchangeSkeletalMeshLodDataNode* LodDataNode = Cast<UInterchangeSkeletalMeshLodDataNode>(Arguments.NodeContainer->GetNode(LodUniqueId));
 				if (!LodDataNode)
 				{
-					UE_LOG(LogInterchangeImportPlugin, Warning, TEXT("Invalid LOD when importing SkeletalMesh asset %s"), *Arguments.AssetName);
+					UE_LOG(LogInterchangeImport, Warning, TEXT("Invalid LOD when importing SkeletalMesh asset %s"), *Arguments.AssetName);
 					continue;
 				}
 				FString SkeletonNodeId;
 				if (!LodDataNode->GetCustomSkeletonID(SkeletonNodeId))
 				{
-					UE_LOG(LogInterchangeImportPlugin, Warning, TEXT("Invalid Skeleton LOD when importing SkeletalMesh asset %s"), *Arguments.AssetName);
+					UE_LOG(LogInterchangeImport, Warning, TEXT("Invalid Skeleton LOD when importing SkeletalMesh asset %s"), *Arguments.AssetName);
 					continue;
 				}
 				const UInterchangeSkeletonNode* SkeletonNode = Cast<UInterchangeSkeletonNode>(Arguments.NodeContainer->GetNode(SkeletonNodeId));
 				if (!SkeletonNode)
 				{
-					UE_LOG(LogInterchangeImportPlugin, Warning, TEXT("Invalid Skeleton LOD when importing SkeletalMesh asset %s"), *Arguments.AssetName);
+					UE_LOG(LogInterchangeImport, Warning, TEXT("Invalid Skeleton LOD when importing SkeletalMesh asset %s"), *Arguments.AssetName);
 					continue;
 				}
 				
@@ -1153,7 +1153,7 @@ UObject* UInterchangeSkeletalMeshFactory::CreateAsset(const UInterchangeSkeletal
 					}
 					if (!ensure(SkeletonReference))
 					{
-						UE_LOG(LogInterchangeImportPlugin, Warning, TEXT("Invalid Skeleton LOD when importing SkeletalMesh asset %s"), *Arguments.AssetName);
+						UE_LOG(LogInterchangeImport, Warning, TEXT("Invalid Skeleton LOD when importing SkeletalMesh asset %s"), *Arguments.AssetName);
 						break;
 					}
 				}
@@ -1161,7 +1161,7 @@ UObject* UInterchangeSkeletalMeshFactory::CreateAsset(const UInterchangeSkeletal
 				FString RootJointNodeId;
 				if (!SkeletonNode->GetCustomRootJointID(RootJointNodeId))
 				{
-					UE_LOG(LogInterchangeImportPlugin, Warning, TEXT("Invalid Skeleton LOD Root Joint when importing SkeletalMesh asset %s"), *Arguments.AssetName);
+					UE_LOG(LogInterchangeImport, Warning, TEXT("Invalid Skeleton LOD Root Joint when importing SkeletalMesh asset %s"), *Arguments.AssetName);
 					continue;
 				}
 				
@@ -1183,7 +1183,7 @@ UObject* UInterchangeSkeletalMeshFactory::CreateAsset(const UInterchangeSkeletal
 					TOptional<UE::Interchange::FSkeletalMeshLodPayloadData> LodMeshPayload = SkeletalMeshTranslatorPayloadInterface->GetSkeletalMeshLodPayloadData(TranslatorMeshKeys[MeshLodIndex]);
 					if (!LodMeshPayload.IsSet())
 					{
-						UE_LOG(LogInterchangeImportPlugin, Warning, TEXT("Invalid Skeletal mesh payload key [%s] SkeletalMesh asset %s"), *TranslatorMeshKeys[MeshLodIndex], *Arguments.AssetName);
+						UE_LOG(LogInterchangeImport, Warning, TEXT("Invalid Skeletal mesh payload key [%s] SkeletalMesh asset %s"), *TranslatorMeshKeys[MeshLodIndex], *Arguments.AssetName);
 						break;
 					}
 					//TODO use the mesh description as the source import data for skeletalmesh
diff --git a/Engine/Plugins/Runtime/Interchange/Source/Import/Private/Mesh/InterchangeSkeletonFactory.cpp b/Engine/Plugins/Runtime/Interchange/Source/Import/Private/Mesh/InterchangeSkeletonFactory.cpp
--- a/Engine/Plugins/Runtime/Interchange/Source/Import/Private/Mesh/InterchangeSkeletonFactory.cpp
+++ b/Engine/Plugins/Runtime/Interchange/Source/Import/Private/Mesh/InterchangeSkeletonFactory.cpp
@@ -3,9 +3,9 @@
 
 #include "Animation/Skeleton.h"
 #include "InterchangeImportCommon.h"
+#include "InterchangeImportLog.h"
 #include "InterchangeSkeletonNode.h"
 #include "InterchangeSourceData.h"
-#include "LogInterchangeImportPlugin.h"
 #include "Nodes/InterchangeBaseNode.h"
 #include "Nodes/InterchangeBaseNodeContainer.h"
 
@@ -48,7 +48,7 @@ UObject* UInterchangeSkeletonFactory::CreateEmptyAsset(const FCreateAssetParams&
 
 	if (!Skeleton)
 	{
-		UE_LOG(LogInterchangeImportPlugin, Warning, TEXT("Could not create Skeleton asset %s"), *Arguments.AssetName);
+		UE_LOG(LogInterchangeImport, Warning, TEXT("Could not create Skeleton asset %s"), *Arguments.AssetName);
 		return nullptr;
 	}
 	Skeleton->PreEditChange(nullptr);
@@ -60,7 +60,7 @@ UObject* UInterchangeSkeletonFactory::CreateAsset(const UInterchangeSkeletonFact
 {
 #if !WITH_EDITORONLY_DATA
 
-	UE_LOG(LogInterchangeImportPlugin, Error, TEXT("Cannot import Skeleton asset in runtime, this is an editor only feature."));
+	UE_LOG(LogInterchangeImport, Error, TEXT("Cannot import Skeleton asset in runtime, this is an editor only feature."));
 	return nullptr;
 
 #else
@@ -99,7 +99,7 @@ UObject* UInterchangeSkeletonFactory::CreateAsset(const UInterchangeSkeletonFact
 
 	if (!SkeletonObject)
 	{
-		UE_LOG(LogInterchangeImportPlugin, Warning, TEXT("Could not create Skeleton asset %s"), *Arguments.AssetName);
+		UE_LOG(LogInterchangeImport, Warning, TEXT("Could not create Skeleton asset %s"), *Arguments.AssetName);
 		return nullptr;
 	}
 
@@ -112,7 +112,7 @@ UObject* UInterchangeSkeletonFactory::CreateAsset(const UInterchangeSkeletonFact
 			USkeleton* Skeleton = Cast<USkeleton>(SkeletonObject);
 			if (!ensure(Skeleton))
 			{
-				UE_LOG(LogInterchangeImportPlugin, Warning, TEXT("Could not create Skeleton asset %s"), *Arguments.AssetName);
+				UE_LOG(LogInterchangeImport, Warning, TEXT("Could not create Skeleton asset %s"), *Arguments.AssetName);
 				return nullptr;
 			}
 		}
diff --git a/Engine/Plugins/Runtime/Interchange/Source/Import/Private/Texture/InterchangeBMPTranslator.cpp b/Engine/Plugins/Runtime/Interchange/Source/Import/Private/Texture/InterchangeBMPTranslator.cpp
--- a/Engine/Plugins/Runtime/Interchange/Source/Import/Private/Texture/InterchangeBMPTranslator.cpp
+++ b/Engine/Plugins/Runtime/Interchange/Source/Import/Private/Texture/InterchangeBMPTranslator.cpp
@@ -6,8 +6,8 @@
 #include "Engine/Texture2D.h"
 #include "IImageWrapper.h"
 #include "IImageWrapperModule.h"
+#include "InterchangeImportLog.h"
 #include "InterchangeTextureNode.h"
-#include "LogInterchangeImportPlugin.h"
 #include "Misc/ConfigCacheIni.h"
 #include "Misc/FileHelper.h"
 #include "Misc/Paths.h"
@@ -33,7 +33,7 @@ TOptional<UE::Interchange::FImportImage> UInterchangeBMPTranslator::GetTexturePa
 {
 	if (!SourceData)
 	{
-		UE_LOG(LogInterchangeImportPlugin, Error, TEXT("Failed to import BMP, bad source data."));
+		UE_LOG(LogInterchangeImport, Error, TEXT("Failed to import BMP, bad source data."));
 		return TOptional<UE::Interchange::FImportImage>();
 	}
 
@@ -43,19 +43,19 @@ TOptional<UE::Interchange::FImportImage> UInterchangeBMPTranslator::GetTexturePa
 	//Make sure the key fit the filename, The key should always be valid
 	if (!Filename.Equals(PayLoadKey))
 	{
-		UE_LOG(LogInterchangeImportPlugin, Error, TEXT("Failed to import BMP, wrong payload key. [%s]"), *Filename);
+		UE_LOG(LogInterchangeImport, Error, TEXT("Failed to import BMP, wrong payload key. [%s]"), *Filename);
 		return TOptional<UE::Interchange::FImportImage>();
 	}
 
 	if (!FPaths::FileExists(Filename))
 	{
-		UE_LOG(LogInterchangeImportPlugin, Error, TEXT("Failed to import BMP, cannot open file. [%s]"), *Filename);
+		UE_LOG(LogInterchangeImport, Error, TEXT("Failed to import BMP, cannot open file. [%s]"), *Filename);
 		return TOptional<UE::Interchange::FImportImage>();
 	}
 
 	if (!FFileHelper::LoadFileToArray(SourceDataBuffer, *Filename))
 	{
-		UE_LOG(LogInterchangeImportPlugin, Error, TEXT("Failed to import BMP, cannot load file content into an array. [%s]"), *Filename);
+		UE_LOG(LogInterchangeImport, Error, TEXT("Failed to import BMP, cannot load file content into an array. [%s]"), *Filename);
 		return TOptional<UE::Interchange::FImportImage>();
 	}
 
@@ -76,12 +76,12 @@ TOptional<UE::Interchange::FImportImage> UInterchangeBMPTranslator::GetTexturePa
 	TSharedPtr<IImageWrapper> BmpImageWrapper = ImageWrapperModule.CreateImageWrapper(EImageFormat::BMP);
 	if (!BmpImageWrapper.IsValid() || !BmpImageWrapper->SetCompressed(Buffer, Length))
 	{
-		UE_LOG(LogInterchangeImportPlugin, Error, TEXT("Failed to decode BMP. [%s]"), *Filename);
+		UE_LOG(LogInterchangeImport, Error, TEXT("Failed to decode BMP. [%s]"), *Filename);
 		return TOptional<UE::Interchange::FImportImage>();
 	}
 	if (!UE::Interchange::FImportImageHelper::IsImportResolutionValid(BmpImageWrapper->GetWidth(), BmpImageWrapper->GetHeight(), bAllowNonPowerOfTwo))
 	{
-		UE_LOG(LogInterchangeImportPlugin, Error, TEXT("Failed to import BMP, invalid resolution. Resolution[%d, %d], AllowPowerOfTwo[%s], [%s]"), BmpImageWrapper->GetWidth(), BmpImageWrapper->GetHeight(), bAllowNonPowerOfTwo ? TEXT("True") : TEXT("false"), *Filename);
+		UE_LOG(LogInterchangeImport, Error, TEXT("Failed to import BMP, invalid resolution. Resolution[%d, %d], AllowPowerOfTwo[%s], [%s]"), BmpImageWrapper->GetWidth(), BmpImageWrapper->GetHeight(), bAllowNonPowerOfTwo ? TEXT("True") : TEXT("false"), *Filename);
 		return TOptional<UE::Interchange::FImportImage>();
 	}
 
@@ -99,7 +99,7 @@ TOptional<UE::Interchange::FImportImage> UInterchangeBMPTranslator::GetTexturePa
 
 	if (!BmpImageWrapper->GetRaw(Format, BitDepth, PayloadData.RawData))
 	{
-		UE_LOG(LogInterchangeImportPlugin, Error, TEXT("Failed to decode BMP. [%s]"), *Filename);
+		UE_LOG(LogInterchangeImport, Error, TEXT("Failed to decode BMP. [%s]"), *Filename);
 		return TOptional<UE::Interchange::FImportImage>();
 	}
 
diff --git a/Engine/Plugins/Runtime/Interchange/Source/Import/Private/Texture/InterchangeDDSTranslator.cpp b/Engine/Plugins/Runtime/Interchange/Source/Import/Private/Texture/InterchangeDDSTranslator.cpp
--- a/Engine/Plugins/Runtime/Interchange/Source/Import/Private/Texture/InterchangeDDSTranslator.cpp
+++ b/Engine/Plugins/Runtime/Interchange/Source/Import/Private/Texture/InterchangeDDSTranslator.cpp
@@ -6,8 +6,8 @@
 #include "Engine/Texture2D.h"
 #include "IImageWrapper.h"
 #include "IImageWrapperModule.h"
+#include "InterchangeImportLog.h"
 #include "InterchangeTextureNode.h"
-#include "LogInterchangeImportPlugin.h"
 #include "Misc/ConfigCacheIni.h"
 #include "Misc/FileHelper.h"
 #include "Misc/Paths.h"
@@ -64,7 +64,7 @@ TOptional<UE::Interchange::FImportImage> UInterchangeDDSTranslator::GetTexturePa
 {
 	if (!SourceData)
 	{
-		UE_LOG(LogInterchangeImportPlugin, Error, TEXT("Failed to import DDS, bad source data."));
+		UE_LOG(LogInterchangeImport, Error, TEXT("Failed to import DDS, bad source data."));
 		return TOptional<UE::Interchange::FImportImage>();
 	}
 
@@ -74,19 +74,19 @@ TOptional<UE::Interchange::FImportImage> UInterchangeDDSTranslator::GetTexturePa
 	//Make sure the key fit the filename, The key should always be valid
 	if (!Filename.Equals(PayLoadKey))
 	{
-		UE_LOG(LogInterchangeImportPlugin, Error, TEXT("Failed to import DDS, wrong payload key. [%s]"), *Filename);
+		UE_LOG(LogInterchangeImport, Error, TEXT("Failed to import DDS, wrong payload key. [%s]"), *Filename);
 		return TOptional<UE::Interchange::FImportImage>();
 	}
 
 	if (!FPaths::FileExists(Filename))
 	{
-		UE_LOG(LogInterchangeImportPlugin, Error, TEXT("Failed to import DDS, cannot open file. [%s]"), *Filename);
+		UE_LOG(LogInterchangeImport, Error, TEXT("Failed to import DDS, cannot open file. [%s]"), *Filename);
 		return TOptional<UE::Interchange::FImportImage>();
 	}
 
 	if (!FFileHelper::LoadFileToArray(SourceDataBuffer, *Filename))
 	{
-		UE_LOG(LogInterchangeImportPlugin, Error, TEXT("Failed to import DDS, cannot load file content into an array. [%s]"), *Filename);
+		UE_LOG(LogInterchangeImport, Error, TEXT("Failed to import DDS, cannot load file content into an array. [%s]"), *Filename);
 		return TOptional<UE::Interchange::FImportImage>();
 	}
 
@@ -106,13 +106,13 @@ TOptional<UE::Interchange::FImportImage> UInterchangeDDSTranslator::GetTexturePa
 	FDDSLoadHelper  DDSLoadHelper(Buffer, Length);
 	if (!DDSLoadHelper.IsValid2DTexture())
 	{
-		UE_LOG(LogInterchangeImportPlugin, Error, TEXT("Failed to import DDS, unsupported format. [%s]"), *Filename);
+		UE_LOG(LogInterchangeImport, Error, TEXT("Failed to import DDS, unsupported format. [%s]"), *Filename);
 		return TOptional<UE::Interchange::FImportImage>();
 	}
 	// DDS 2d texture
 	if (!UE::Interchange::FImportImageHelper::IsImportResolutionValid(DDSLoadHelper.DDSHeader->dwWidth, DDSLoadHelper.DDSHeader->dwHeight, bAllowNonPowerOfTwo))
 	{
-		UE_LOG(LogInterchangeImportPlugin, Error, TEXT("Failed to import PCX, invalid resolution. Resolution[%d, %d], AllowPowerOfTwo[%s], [%s]"), DDSLoadHelper.DDSHeader->dwWidth, DDSLoadHelper.DDSHeader->dwHeight, bAllowNonPowerOfTwo ? TEXT("True") : TEXT("false"), *Filename);
+		UE_LOG(LogInterchangeImport, Error, TEXT("Failed to import PCX, invalid resolution. Resolution[%d, %d], AllowPowerOfTwo[%s], [%s]"), DDSLoadHelper.DDSHeader->dwWidth, DDSLoadHelper.DDSHeader->dwHeight, bAllowNonPowerOfTwo ? TEXT("True") : TEXT("false"), *Filename);
 		return TOptional<UE::Interchange::FImportImage>();
 	}
 
@@ -121,14 +121,14 @@ TOptional<UE::Interchange::FImportImage> UInterchangeDDSTranslator::GetTexturePa
 	// Invalid DDS format
 	if (SourceFormat == TSF_Invalid)
 	{
-		UE_LOG(LogInterchangeImportPlugin, Error, TEXT("DDS file [%s] contains data in an unsupported format"), *Filename);
+		UE_LOG(LogInterchangeImport, Error, TEXT("DDS file [%s] contains data in an unsupported format"), *Filename);
 		return TOptional<UE::Interchange::FImportImage>();
 	}
 
 	uint32 MipMapCount = DDSLoadHelper.ComputeMipMapCount();
 	if (MipMapCount <= 0)
 	{
-		UE_LOG(LogInterchangeImportPlugin, Error, TEXT("DDS file [%s] do not have any mipmap"), *Filename);
+		UE_LOG(LogInterchangeImport, Error, TEXT("DDS file [%s] do not have any mipmap"), *Filename);
 		return TOptional<UE::Interchange::FImportImage>();
 	}
 
diff --git a/Engine/Plugins/Runtime/Interchange/Source/Import/Private/Texture/InterchangeEXRTranslator.cpp b/Engine/Plugins/Runtime/Interchange/Source/Import/Private/Texture/InterchangeEXRTranslator.cpp
--- a/Engine/Plugins/Runtime/Interchange/Source/Import/Private/Texture/InterchangeEXRTranslator.cpp
+++ b/Engine/Plugins/Runtime/Interchange/Source/Import/Private/Texture/InterchangeEXRTranslator.cpp
@@ -5,8 +5,8 @@
 #include "Engine/Texture2D.h"
 #include "IImageWrapper.h"
 #include "IImageWrapperModule.h"
+#include "InterchangeImportLog.h"
 #include "InterchangeTextureNode.h"
-#include "LogInterchangeImportPlugin.h"
 #include "Misc/ConfigCacheIni.h"
 #include "Misc/FileHelper.h"
 #include "Misc/Paths.h"
@@ -31,7 +31,7 @@ TOptional<UE::Interchange::FImportImage> UInterchangeEXRTranslator::GetTexturePa
 {
 	if (!SourceData)
 	{
-		UE_LOG(LogInterchangeImportPlugin, Error, TEXT("Failed to import EXR, bad source data."));
+		UE_LOG(LogInterchangeImport, Error, TEXT("Failed to import EXR, bad source data."));
 		return TOptional<UE::Interchange::FImportImage>();
 	}
 
@@ -41,19 +41,19 @@ TOptional<UE::Interchange::FImportImage> UInterchangeEXRTranslator::GetTexturePa
 	//Make sure the key fit the filename, The key should always be valid
 	if (!Filename.Equals(PayLoadKey))
 	{
-		UE_LOG(LogInterchangeImportPlugin, Error, TEXT("Failed to import EXR, wrong payload key. [%s]"), *Filename);
+		UE_LOG(LogInterchangeImport, Error, TEXT("Failed to import EXR, wrong payload key. [%s]"), *Filename);
 		return TOptional<UE::Interchange::FImportImage>();
 	}
 
 	if (!FPaths::FileExists(Filename))
 	{
-		UE_LOG(LogInterchangeImportPlugin, Error, TEXT("Failed to import EXR, cannot open file. [%s]"), *Filename);
+		UE_LOG(LogInterchangeImport, Error, TEXT("Failed to import EXR, cannot open file. [%s]"), *Filename);
 		return TOptional<UE::Interchange::FImportImage>();
 	}
 
 	if (!FFileHelper::LoadFileToArray(SourceDataBuffer, *Filename))
 	{
-		UE_LOG(LogInterchangeImportPlugin, Error, TEXT("Failed to import EXR, cannot load file content into an array. [%s]"), *Filename);
+		UE_LOG(LogInterchangeImport, Error, TEXT("Failed to import EXR, cannot load file content into an array. [%s]"), *Filename);
 		return TOptional<UE::Interchange::FImportImage>();
 	}
 
@@ -74,7 +74,7 @@ TOptional<UE::Interchange::FImportImage> UInterchangeEXRTranslator::GetTexturePa
 	TSharedPtr<IImageWrapper> ExrImageWrapper = ImageWrapperModule.CreateImageWrapper(EImageFormat::EXR);
 	if (!ExrImageWrapper.IsValid() || !ExrImageWrapper->SetCompressed(Buffer, Length))
 	{
-		UE_LOG(LogInterchangeImportPlugin, Error, TEXT("Failed to import EXR, unsupported format. [%s]"), *Filename);
+		UE_LOG(LogInterchangeImport, Error, TEXT("Failed to import EXR, unsupported format. [%s]"), *Filename);
 		return TOptional<UE::Interchange::FImportImage>();
 	}
 
@@ -83,7 +83,7 @@ TOptional<UE::Interchange::FImportImage> UInterchangeEXRTranslator::GetTexturePa
 
 	if (!UE::Interchange::FImportImageHelper::IsImportResolutionValid(Width, Height, bAllowNonPowerOfTwo))
 	{
-		UE_LOG(LogInterchangeImportPlugin, Error, TEXT("Failed to import PCX, invalid resolution. Resolution[%d, %d], AllowPowerOfTwo[%s], [%s]"), Width, Height, bAllowNonPowerOfTwo ? TEXT("True") : TEXT("false"), *Filename);
+		UE_LOG(LogInterchangeImport, Error, TEXT("Failed to import PCX, invalid resolution. Resolution[%d, %d], AllowPowerOfTwo[%s], [%s]"), Width, Height, bAllowNonPowerOfTwo ? TEXT("True") : TEXT("false"), *Filename);
 		return TOptional<UE::Interchange::FImportImage>();
 	}
 
@@ -100,7 +100,7 @@ TOptional<UE::Interchange::FImportImage> UInterchangeEXRTranslator::GetTexturePa
 
 	if (TextureFormat == TSF_Invalid)
 	{
-		UE_LOG(LogInterchangeImportPlugin, Error, TEXT("EXR file [%s] contains data in an unsupported format"), *Filename);
+		UE_LOG(LogInterchangeImport, Error, TEXT("EXR file [%s] contains data in an unsupported format"), *Filename);
 		return TOptional<UE::Interchange::FImportImage>();
 	}
 
@@ -115,7 +115,7 @@ TOptional<UE::Interchange::FImportImage> UInterchangeEXRTranslator::GetTexturePa
 
 	if (!ExrImageWrapper->GetRaw(Format, BitDepth, PayloadData.RawData))
 	{
-		UE_LOG(LogInterchangeImportPlugin, Error, TEXT("Failed to decode EXR. [%s]"), *Filename);
+		UE_LOG(LogInterchangeImport, Error, TEXT("Failed to decode EXR. [%s]"), *Filename);
 		return TOptional<UE::Interchange::FImportImage>();
 	}
 
diff --git a/Engine/Plugins/Runtime/Interchange/Source/Import/Private/Texture/InterchangeJPGTranslator.cpp b/Engine/Plugins/Runtime/Interchange/Source/Import/Private/Texture/InterchangeJPGTranslator.cpp
--- a/Engine/Plugins/Runtime/Interchange/Source/Import/Private/Texture/InterchangeJPGTranslator.cpp
+++ b/Engine/Plugins/Runtime/Interchange/Source/Import/Private/Texture/InterchangeJPGTranslator.cpp
@@ -5,8 +5,8 @@
 #include "Engine/Texture2D.h"
 #include "IImageWrapper.h"
 #include "IImageWrapperModule.h"
+#include "InterchangeImportLog.h"
 #include "InterchangeTextureNode.h"
-#include "LogInterchangeImportPlugin.h"
 #include "Misc/ConfigCacheIni.h"
 #include "Misc/FileHelper.h"
 #include "Misc/Paths.h"
@@ -31,7 +31,7 @@ TOptional<UE::Interchange::FImportImage> UInterchangeJPGTranslator::GetTexturePa
 {
 	if (!SourceData)
 	{
-		UE_LOG(LogInterchangeImportPlugin, Error, TEXT("Failed to import JPEG, bad source data."));
+		UE_LOG(LogInterchangeImport, Error, TEXT("Failed to import JPEG, bad source data."));
 		return TOptional<UE::Interchange::FImportImage>();
 	}
 
@@ -41,19 +41,19 @@ TOptional<UE::Interchange::FImportImage> UInterchangeJPGTranslator::GetTexturePa
 	//Make sure the key fit the filename, The key should always be valid
 	if (!Filename.Equals(PayLoadKey))
 	{
-		UE_LOG(LogInterchangeImportPlugin, Error, TEXT("Failed to import JPEG, wrong payload key. [%s]"), *Filename);
+		UE_LOG(LogInterchangeImport, Error, TEXT("Failed to import JPEG, wrong payload key. [%s]"), *Filename);
 		return TOptional<UE::Interchange::FImportImage>();
 	}
 
 	if (!FPaths::FileExists(Filename))
 	{
-		UE_LOG(LogInterchangeImportPlugin, Error, TEXT("Failed to import JPEG, cannot open file. [%s]"), *Filename);
+		UE_LOG(LogInterchangeImport, Error, TEXT("Failed to import JPEG, cannot open file. [%s]"), *Filename);
 		return TOptional<UE::Interchange::FImportImage>();
 	}
 
 	if (!FFileHelper::LoadFileToArray(SourceDataBuffer, *Filename))
 	{
-		UE_LOG(LogInterchangeImportPlugin, Error, TEXT("Failed to import JPEG, cannot load file content into an array. [%s]"), *Filename);
+		UE_LOG(LogInterchangeImport, Error, TEXT("Failed to import JPEG, cannot load file content into an array. [%s]"), *Filename);
 		return TOptional<UE::Interchange::FImportImage>();
 	}
 
@@ -74,12 +74,12 @@ TOptional<UE::Interchange::FImportImage> UInterchangeJPGTranslator::GetTexturePa
 	TSharedPtr<IImageWrapper> JpegImageWrapper = ImageWrapperModule.CreateImageWrapper(EImageFormat::JPEG);
 	if (!JpegImageWrapper.IsValid() || !JpegImageWrapper->SetCompressed(Buffer, Length))
 	{
-		UE_LOG(LogInterchangeImportPlugin, Error, TEXT("Failed to decode JPEG. [%s]"), *Filename);
+		UE_LOG(LogInterchangeImport, Error, TEXT("Failed to decode JPEG. [%s]"), *Filename);
 		return TOptional<UE::Interchange::FImportImage>();
 	}
 	if (!UE::Interchange::FImportImageHelper::IsImportResolutionValid(JpegImageWrapper->GetWidth(), JpegImageWrapper->GetHeight(), bAllowNonPowerOfTwo))
 	{
-		UE_LOG(LogInterchangeImportPlugin, Error, TEXT("Failed to import JPEG, invalid resolution. Resolution[%d, %d], AllowPowerOfTwo[%s], [%s]"), JpegImageWrapper->GetWidth(), JpegImageWrapper->GetHeight(), bAllowNonPowerOfTwo ? TEXT("True") : TEXT("false"), *Filename);
+		UE_LOG(LogInterchangeImport, Error, TEXT("Failed to import JPEG, invalid resolution. Resolution[%d, %d], AllowPowerOfTwo[%s], [%s]"), JpegImageWrapper->GetWidth(), JpegImageWrapper->GetHeight(), bAllowNonPowerOfTwo ? TEXT("True") : TEXT("false"), *Filename);
 		return TOptional<UE::Interchange::FImportImage>();
 	}
 
@@ -109,7 +109,7 @@ TOptional<UE::Interchange::FImportImage> UInterchangeJPGTranslator::GetTexturePa
 
 	if (TextureFormat == TSF_Invalid)
 	{
-		UE_LOG(LogInterchangeImportPlugin, Error, TEXT("JPEG file [%s] contains data in an unsupported format"), *Filename);
+		UE_LOG(LogInterchangeImport, Error, TEXT("JPEG file [%s] contains data in an unsupported format"), *Filename);
 		return TOptional<UE::Interchange::FImportImage>();
 	}
 
@@ -124,7 +124,7 @@ TOptional<UE::Interchange::FImportImage> UInterchangeJPGTranslator::GetTexturePa
 
 	if (!JpegImageWrapper->GetRaw(Format, BitDepth, PayloadData.RawData))
 	{
-		UE_LOG(LogInterchangeImportPlugin, Error, TEXT("Failed to decode JPEG. [%s]"), *Filename);
+		UE_LOG(LogInterchangeImport, Error, TEXT("Failed to decode JPEG. [%s]"), *Filename);
 		return TOptional<UE::Interchange::FImportImage>();
 	}
 	return PayloadData;
diff --git a/Engine/Plugins/Runtime/Interchange/Source/Import/Private/Texture/InterchangePCXTranslator.cpp b/Engine/Plugins/Runtime/Interchange/Source/Import/Private/Texture/InterchangePCXTranslator.cpp
--- a/Engine/Plugins/Runtime/Interchange/Source/Import/Private/Texture/InterchangePCXTranslator.cpp
+++ b/Engine/Plugins/Runtime/Interchange/Source/Import/Private/Texture/InterchangePCXTranslator.cpp
@@ -5,8 +5,8 @@
 #include "Engine/Texture2D.h"
 #include "IImageWrapper.h"
 #include "IImageWrapperModule.h"
+#include "InterchangeImportLog.h"
 #include "InterchangeTextureNode.h"
-#include "LogInterchangeImportPlugin.h"
 #include "Misc/ConfigCacheIni.h"
 #include "Misc/FileHelper.h"
 #include "Misc/Paths.h"
@@ -82,7 +82,7 @@ TOptional<UE::Interchange::FImportImage> UInterchangePCXTranslator::GetTexturePa
 {
 	if (!SourceData)
 	{
-		UE_LOG(LogInterchangeImportPlugin, Error, TEXT("Failed to import PCX, bad source data."));
+		UE_LOG(LogInterchangeImport, Error, TEXT("Failed to import PCX, bad source data."));
 		return TOptional<UE::Interchange::FImportImage>();
 	}
 
@@ -92,19 +92,19 @@ TOptional<UE::Interchange::FImportImage> UInterchangePCXTranslator::GetTexturePa
 	//Make sure the key fit the filename, The key should always be valid
 	if (!Filename.Equals(PayLoadKey))
 	{
-		UE_LOG(LogInterchangeImportPlugin, Error, TEXT("Failed to import PCX, wrong payload key. [%s]"), *Filename);
+		UE_LOG(LogInterchangeImport, Error, TEXT("Failed to import PCX, wrong payload key. [%s]"), *Filename);
 		return TOptional<UE::Interchange::FImportImage>();
 	}
 
 	if (!FPaths::FileExists(Filename))
 	{
-		UE_LOG(LogInterchangeImportPlugin, Error, TEXT("Failed to import PCX, cannot open file. [%s]"), *Filename);
+		UE_LOG(LogInterchangeImport, Error, TEXT("Failed to import PCX, cannot open file. [%s]"), *Filename);
 		return TOptional<UE::Interchange::FImportImage>();
 	}
 
 	if (!FFileHelper::LoadFileToArray(SourceDataBuffer, *Filename))
 	{
-		UE_LOG(LogInterchangeImportPlugin, Error, TEXT("Failed to import PCX, cannot load file content into an array. [%s]"), *Filename);
+		UE_LOG(LogInterchangeImport, Error, TEXT("Failed to import PCX, cannot load file content into an array. [%s]"), *Filename);
 		return TOptional<UE::Interchange::FImportImage>();
 	}
 
@@ -133,7 +133,7 @@ TOptional<UE::Interchange::FImportImage> UInterchangePCXTranslator::GetTexturePa
 		// Check the resolution of the imported texture to ensure validity
 		if (!UE::Interchange::FImportImageHelper::IsImportResolutionValid(NewU, NewV, bAllowNonPowerOfTwo))
 		{
-			UE_LOG(LogInterchangeImportPlugin, Error, TEXT("Failed to import PCX, invalid resolution. Resolution[%d, %d], AllowPowerOfTwo[%s], [%s]"), NewU, NewV, bAllowNonPowerOfTwo ? TEXT("True") : TEXT("false"), *Filename);
+			UE_LOG(LogInterchangeImport, Error, TEXT("Failed to import PCX, invalid resolution. Resolution[%d, %d], AllowPowerOfTwo[%s], [%s]"), NewU, NewV, bAllowNonPowerOfTwo ? TEXT("True") : TEXT("false"), *Filename);
 			return TOptional<UE::Interchange::FImportImage>();
 		}
 		else if (PCX->NumPlanes == 1 && PCX->BitsPerPixel == 8)
@@ -231,13 +231,13 @@ TOptional<UE::Interchange::FImportImage> UInterchangePCXTranslator::GetTexturePa
 		}
 		else
 		{
-			UE_LOG(LogInterchangeImportPlugin, Error, TEXT("Failed to import PCX, unsupported format. (%i/%i) [%s]"), PCX->NumPlanes, PCX->BitsPerPixel, *Filename);
+			UE_LOG(LogInterchangeImport, Error, TEXT("Failed to import PCX, unsupported format. (%i/%i) [%s]"), PCX->NumPlanes, PCX->BitsPerPixel, *Filename);
 			return TOptional<UE::Interchange::FImportImage>();
 		}
 	}
 	else
 	{
-		UE_LOG(LogInterchangeImportPlugin, Error, TEXT("Failed to import PCX, unsupported file version. [%s]"), *Filename);
+		UE_LOG(LogInterchangeImport, Error, TEXT("Failed to import PCX, unsupported file version. [%s]"), *Filename);
 		return TOptional<UE::Interchange::FImportImage>();
 	}
 
diff --git a/Engine/Plugins/Runtime/Interchange/Source/Import/Private/Texture/InterchangePNGTranslator.cpp b/Engine/Plugins/Runtime/Interchange/Source/Import/Private/Texture/InterchangePNGTranslator.cpp
--- a/Engine/Plugins/Runtime/Interchange/Source/Import/Private/Texture/InterchangePNGTranslator.cpp
+++ b/Engine/Plugins/Runtime/Interchange/Source/Import/Private/Texture/InterchangePNGTranslator.cpp
@@ -5,8 +5,8 @@
 #include "Engine/Texture2D.h"
 #include "IImageWrapper.h"
 #include "IImageWrapperModule.h"
+#include "InterchangeImportLog.h"
 #include "InterchangeTextureNode.h"
-#include "LogInterchangeImportPlugin.h"
 #include "Misc/ConfigCacheIni.h"
 #include "Misc/FileHelper.h"
 #include "Misc/Paths.h"
@@ -203,7 +203,7 @@ TOptional<UE::Interchange::FImportImage> UInterchangePNGTranslator::GetTexturePa
 {
 	if (!SourceData)
 	{
-		UE_LOG(LogInterchangeImportPlugin, Error, TEXT("Failed to import PNG, bad source data."));
+		UE_LOG(LogInterchangeImport, Error, TEXT("Failed to import PNG, bad source data."));
 		return TOptional<UE::Interchange::FImportImage>();
 	}
 
@@ -213,19 +213,19 @@ TOptional<UE::Interchange::FImportImage> UInterchangePNGTranslator::GetTexturePa
 	//Make sure the key fit the filename, The key should always be valid
 	if (!Filename.Equals(PayLoadKey))
 	{
-		UE_LOG(LogInterchangeImportPlugin, Error, TEXT("Failed to import PNG, wrong payload key. [%s]"), *Filename);
+		UE_LOG(LogInterchangeImport, Error, TEXT("Failed to import PNG, wrong payload key. [%s]"), *Filename);
 		return TOptional<UE::Interchange::FImportImage>();
 	}
 
 	if (!FPaths::FileExists(Filename))
 	{
-		UE_LOG(LogInterchangeImportPlugin, Error, TEXT("Failed to import PNG, cannot open file. [%s]"), *Filename);
+		UE_LOG(LogInterchangeImport, Error, TEXT("Failed to import PNG, cannot open file. [%s]"), *Filename);
 		return TOptional<UE::Interchange::FImportImage>();
 	}
 
 	if (!FFileHelper::LoadFileToArray(SourceDataBuffer, *Filename))
 	{
-		UE_LOG(LogInterchangeImportPlugin, Error, TEXT("Failed to import PNG, cannot load file content into an array. [%s]"), *Filename);
+		UE_LOG(LogInterchangeImport, Error, TEXT("Failed to import PNG, cannot load file content into an array. [%s]"), *Filename);
 		return TOptional<UE::Interchange::FImportImage>();
 	}
 
@@ -246,12 +246,12 @@ TOptional<UE::Interchange::FImportImage> UInterchangePNGTranslator::GetTexturePa
 	TSharedPtr<IImageWrapper> PngImageWrapper = ImageWrapperModule.CreateImageWrapper(EImageFormat::PNG);
 	if (!PngImageWrapper.IsValid() || !PngImageWrapper->SetCompressed(Buffer, Length))
 	{
-		UE_LOG(LogInterchangeImportPlugin, Error, TEXT("Failed to decode PNG. [%s]"), *Filename);
+		UE_LOG(LogInterchangeImport, Error, TEXT("Failed to decode PNG. [%s]"), *Filename);
 		return TOptional<UE::Interchange::FImportImage>();
 	}
 	if (!UE::Interchange::FImportImageHelper::IsImportResolutionValid(PngImageWrapper->GetWidth(), PngImageWrapper->GetHeight(), bAllowNonPowerOfTwo))
 	{
-		UE_LOG(LogInterchangeImportPlugin, Error, TEXT("Failed to import PNG, invalid resolution. Resolution[%d, %d], AllowPowerOfTwo[%s], [%s]"), PngImageWrapper->GetWidth(), PngImageWrapper->GetHeight(), bAllowNonPowerOfTwo ? TEXT("True") : TEXT("false"), *Filename);
+		UE_LOG(LogInterchangeImport, Error, TEXT("Failed to import PNG, invalid resolution. Resolution[%d, %d], AllowPowerOfTwo[%s], [%s]"), PngImageWrapper->GetWidth(), PngImageWrapper->GetHeight(), bAllowNonPowerOfTwo ? TEXT("True") : TEXT("false"), *Filename);
 		return TOptional<UE::Interchange::FImportImage>();
 	}
 
@@ -294,7 +294,7 @@ TOptional<UE::Interchange::FImportImage> UInterchangePNGTranslator::GetTexturePa
 
 	if (TextureFormat == TSF_Invalid)
 	{
-		UE_LOG(LogInterchangeImportPlugin, Error, TEXT("PNG file [%s] contains data in an unsupported format"), *Filename);
+		UE_LOG(LogInterchangeImport, Error, TEXT("PNG file [%s] contains data in an unsupported format"), *Filename);
 		return TOptional<UE::Interchange::FImportImage>();
 	}
 
@@ -319,7 +319,7 @@ TOptional<UE::Interchange::FImportImage> UInterchangePNGTranslator::GetTexturePa
 	}
 	else
 	{
-		UE_LOG(LogInterchangeImportPlugin, Error, TEXT("Failed to decode PNG. [%s]"), *Filename);
+		UE_LOG(LogInterchangeImport, Error, TEXT("Failed to decode PNG. [%s]"), *Filename);
 		return TOptional<UE::Interchange::FImportImage>();
 	}
 
diff --git a/Engine/Plugins/Runtime/Interchange/Source/Import/Private/Texture/InterchangePSDTranslator.cpp b/Engine/Plugins/Runtime/Interchange/Source/Import/Private/Texture/InterchangePSDTranslator.cpp
--- a/Engine/Plugins/Runtime/Interchange/Source/Import/Private/Texture/InterchangePSDTranslator.cpp
+++ b/Engine/Plugins/Runtime/Interchange/Source/Import/Private/Texture/InterchangePSDTranslator.cpp
@@ -5,8 +5,8 @@
 #include "Engine/Texture2D.h"
 #include "IImageWrapper.h"
 #include "IImageWrapperModule.h"
+#include "InterchangeImportLog.h"
 #include "InterchangeTextureNode.h"
-#include "LogInterchangeImportPlugin.h"
 #include "Misc/ConfigCacheIni.h"
 #include "Misc/FileHelper.h"
 #include "Misc/Paths.h"
@@ -384,7 +384,7 @@ TOptional<UE::Interchange::FImportImage> UInterchangePSDTranslator::GetTexturePa
 {
 	if (!SourceData)
 	{
-		UE_LOG(LogInterchangeImportPlugin, Error, TEXT("Failed to import PSD, bad source data."));
+		UE_LOG(LogInterchangeImport, Error, TEXT("Failed to import PSD, bad source data."));
 		return TOptional<UE::Interchange::FImportImage>();
 	}
 
@@ -394,19 +394,19 @@ TOptional<UE::Interchange::FImportImage> UInterchangePSDTranslator::GetTexturePa
 	//Make sure the key fit the filename, The key should always be valid
 	if (!Filename.Equals(PayLoadKey))
 	{
-		UE_LOG(LogInterchangeImportPlugin, Error, TEXT("Failed to import PSD, wrong payload key. [%s]"), *Filename);
+		UE_LOG(LogInterchangeImport, Error, TEXT("Failed to import PSD, wrong payload key. [%s]"), *Filename);
 		return TOptional<UE::Interchange::FImportImage>();
 	}
 
 	if (!FPaths::FileExists(Filename))
 	{
-		UE_LOG(LogInterchangeImportPlugin, Error, TEXT("Failed to import PSD, cannot open file. [%s]"), *Filename);
+		UE_LOG(LogInterchangeImport, Error, TEXT("Failed to import PSD, cannot open file. [%s]"), *Filename);
 		return TOptional<UE::Interchange::FImportImage>();
 	}
 
 	if (!FFileHelper::LoadFileToArray(SourceDataBuffer, *Filename))
 	{
-		UE_LOG(LogInterchangeImportPlugin, Error, TEXT("Failed to import PSD, cannot load file content into an array. [%s]"), *Filename);
+		UE_LOG(LogInterchangeImport, Error, TEXT("Failed to import PSD, cannot load file content into an array. [%s]"), *Filename);
 		return TOptional<UE::Interchange::FImportImage>();
 	}
 
@@ -443,7 +443,7 @@ TOptional<UE::Interchange::FImportImage> UInterchangePSDTranslator::GetTexturePa
 
 	if (!PSDHeader.IsSupported())
 	{
-		UE_LOG(LogInterchangeImportPlugin, Error, TEXT("Format of this PSD is not supported. Only Grayscale and RGBColor PSD images are currently supported, in 8-bit or 16-bit."));
+		UE_LOG(LogInterchangeImport, Error, TEXT("Format of this PSD is not supported. Only Grayscale and RGBColor PSD images are currently supported, in 8-bit or 16-bit."));
 		return TOptional<UE::Interchange::FImportImage>();
 	}
 
@@ -460,7 +460,7 @@ TOptional<UE::Interchange::FImportImage> UInterchangePSDTranslator::GetTexturePa
 
 	if (TextureFormat == TSF_Invalid)
 	{
-		UE_LOG(LogInterchangeImportPlugin, Error, TEXT("PSD file contains data in an unsupported format."));
+		UE_LOG(LogInterchangeImport, Error, TEXT("PSD file contains data in an unsupported format."));
 		return TOptional<UE::Interchange::FImportImage>();
 	}
 
@@ -475,7 +475,7 @@ TOptional<UE::Interchange::FImportImage> UInterchangePSDTranslator::GetTexturePa
 
 	if (!UE::Interchange::Private::ReadData(PayloadData.RawData.GetData(), Buffer, PSDHeader))
 	{
-		UE_LOG(LogInterchangeImportPlugin, Error, TEXT("Failed to read this PSD"));
+		UE_LOG(LogInterchangeImport, Error, TEXT("Failed to read this PSD"));
 		return TOptional<UE::Interchange::FImportImage>();
 	}
 
diff --git a/Engine/Plugins/Runtime/Interchange/Source/Import/Private/Texture/InterchangeTGATranslator.cpp b/Engine/Plugins/Runtime/Interchange/Source/Import/Private/Texture/InterchangeTGATranslator.cpp
--- a/Engine/Plugins/Runtime/Interchange/Source/Import/Private/Texture/InterchangeTGATranslator.cpp
+++ b/Engine/Plugins/Runtime/Interchange/Source/Import/Private/Texture/InterchangeTGATranslator.cpp
@@ -3,8 +3,8 @@
 
 #include "IImageWrapper.h"
 #include "IImageWrapperModule.h"
+#include "InterchangeImportLog.h"
 #include "InterchangeTextureNode.h"
-#include "LogInterchangeImportPlugin.h"
 #include "Misc/ConfigCacheIni.h"
 #include "Misc/FileHelper.h"
 #include "Misc/Paths.h"
@@ -47,7 +47,7 @@ namespace TGATranslatorImpl
 					TGA->BitsPerPixel != 24 &&
 					TGA->BitsPerPixel != 16 )
 				{
-					UE_LOG( LogInterchangeImportPlugin, Error, TEXT("TGA uses an unsupported rle-compressed bit-depth: %u"), TGA->BitsPerPixel );
+					UE_LOG( LogInterchangeImport, Error, TEXT("TGA uses an unsupported rle-compressed bit-depth: %u"), TGA->BitsPerPixel );
 					return false;
 				}
 			}
@@ -57,7 +57,7 @@ namespace TGATranslatorImpl
 					TGA->BitsPerPixel != 16 &&
 					TGA->BitsPerPixel != 24)
 				{
-					UE_LOG( LogInterchangeImportPlugin, Error, TEXT("TGA uses an unsupported bit-depth: %u"), TGA->BitsPerPixel );
+					UE_LOG( LogInterchangeImport, Error, TEXT("TGA uses an unsupported bit-depth: %u"), TGA->BitsPerPixel );
 					return false;
 				}
 			}
@@ -88,7 +88,7 @@ TOptional<UE::Interchange::FImportImage> UInterchangeTGATranslator::GetTexturePa
 {
 	if (!SourceData)
 	{
-		UE_LOG(LogInterchangeImportPlugin, Error, TEXT("Failed to import TGA, bad source data."));
+		UE_LOG(LogInterchangeImport, Error, TEXT("Failed to import TGA, bad source data."));
 		return {};
 	}
 
@@ -98,19 +98,19 @@ TOptional<UE::Interchange::FImportImage> UInterchangeTGATranslator::GetTexturePa
 	//Make sure the key fit the filename, The key should always be valid
 	if (!Filename.Equals(PayLoadKey))
 	{
-		UE_LOG(LogInterchangeImportPlugin, Error, TEXT("Failed to import TGA, wrong payload key. [%s]"), *Filename);
+		UE_LOG(LogInterchangeImport, Error, TEXT("Failed to import TGA, wrong payload key. [%s]"), *Filename);
 		return {};
 	}
 
 	if (!FPaths::FileExists(Filename))
 	{
-		UE_LOG(LogInterchangeImportPlugin, Error, TEXT("Failed to import TGA, cannot open file. [%s]"), *Filename);
+		UE_LOG(LogInterchangeImport, Error, TEXT("Failed to import TGA, cannot open file. [%s]"), *Filename);
 		return {};
 	}
 
 	if (!FFileHelper::LoadFileToArray(SourceDataBuffer, *Filename))
 	{
-		UE_LOG(LogInterchangeImportPlugin, Error, TEXT("Failed to import TGA, cannot load file content into an array. [%s]"), *Filename);
+		UE_LOG(LogInterchangeImport, Error, TEXT("Failed to import TGA, cannot load file content into an array. [%s]"), *Filename);
 		return {};
 	}
 
diff --git a/Engine/Plugins/Runtime/Interchange/Source/Import/Private/Texture/InterchangeTextureFactory.cpp b/Engine/Plugins/Runtime/Interchange/Source/Import/Private/Texture/InterchangeTextureFactory.cpp
--- a/Engine/Plugins/Runtime/Interchange/Source/Import/Private/Texture/InterchangeTextureFactory.cpp
+++ b/Engine/Plugins/Runtime/Interchange/Source/Import/Private/Texture/InterchangeTextureFactory.cpp
@@ -5,9 +5,9 @@
 #include "Engine/Texture2D.h"
 #include "InterchangeAssetImportData.h"
 #include "InterchangeImportCommon.h"
+#include "InterchangeImportLog.h"
 #include "InterchangeTextureNode.h"
 #include "InterchangeTranslatorBase.h"
-#include "LogInterchangeImportPlugin.h"
 #include "Nodes/InterchangeBaseNode.h"
 #include "Texture/InterchangeBlockedTexturePayloadInterface.h"
 #include "Texture/InterchangeTexturePayloadInterface.h"
@@ -72,7 +72,7 @@ UObject* UInterchangeTextureFactory::CreateEmptyAsset(const FCreateAssetParams&
 
 	if (!Texture)
 	{
-		UE_LOG(LogInterchangeImportPlugin, Warning, TEXT("Could not create texture asset %s"), *Arguments.AssetName);
+		UE_LOG(LogInterchangeImport, Warning, TEXT("Could not create texture asset %s"), *Arguments.AssetName);
 		return nullptr;
 	}
 
@@ -86,7 +86,7 @@ UObject* UInterchangeTextureFactory::CreateAsset(const UInterchangeTextureFactor
 {
 #if !WITH_EDITORONLY_DATA
 
-	UE_LOG(LogInterchangeImportPlugin, Error, TEXT("Cannot import texture asset in runtime, this is an editor only feature."));
+	UE_LOG(LogInterchangeImport, Error, TEXT("Cannot import texture asset in runtime, this is an editor only feature."));
 	return nullptr;
 
 #else
@@ -151,7 +151,7 @@ UObject* UInterchangeTextureFactory::CreateAsset(const UInterchangeTextureFactor
 
 	if (!Texture)
 	{
-		UE_LOG(LogInterchangeImportPlugin, Warning, TEXT("Could not create texture asset %s"), *Arguments.AssetName);
+		UE_LOG(LogInterchangeImport, Warning, TEXT("Could not create texture asset %s"), *Arguments.AssetName);
 		return nullptr;
 	}
 
diff --git a/Engine/Plugins/Runtime/Interchange/Source/Import/Public/Fbx/InterchangeFbxTranslator.h b/Engine/Plugins/Runtime/Interchange/Source/Import/Public/Fbx/InterchangeFbxTranslator.h
--- a/Engine/Plugins/Runtime/Interchange/Source/Import/Public/Fbx/InterchangeFbxTranslator.h
+++ b/Engine/Plugins/Runtime/Interchange/Source/Import/Public/Fbx/InterchangeFbxTranslator.h
@@ -18,7 +18,7 @@
 /* Fbx translator class support import of texture, material, static mesh, skeletal mesh, */
 
 UCLASS(BlueprintType)
-class INTERCHANGEIMPORTPLUGIN_API UInterchangeFbxTranslator : public UInterchangeTranslatorBase
+class INTERCHANGEIMPORT_API UInterchangeFbxTranslator : public UInterchangeTranslatorBase
 , public IInterchangeTexturePayloadInterface
 , public IInterchangeStaticMeshPayloadInterface
 , public IInterchangeSkeletalMeshPayloadInterface
diff --git a/Engine/Plugins/Runtime/Interchange/Source/Import/Public/IInterchangeImportPlugin.h b/Engine/Plugins/Runtime/Interchange/Source/Import/Public/IInterchangeImportPlugin.h
deleted file mode 100644
diff --git a/Engine/Plugins/Runtime/Interchange/Source/Import/Public/InterchangeImportModule.h b/Engine/Plugins/Runtime/Interchange/Source/Import/Public/InterchangeImportModule.h
new file mode 100644
--- /dev/null
+++ b/Engine/Plugins/Runtime/Interchange/Source/Import/Public/InterchangeImportModule.h
@@ -0,0 +1,38 @@
+// Copyright Epic Games, Inc. All Rights Reserved.
+
+#pragma once
+
+#include "CoreMinimal.h"
+#include "Modules/ModuleInterface.h"
+#include "Modules/ModuleManager.h"
+
+/**
+ * The public interface to this module
+ */
+class IInterchangeImportModule : public IModuleInterface
+{
+
+public:
+
+	/**
+	 * Singleton-like access to this module's interface.  This is just for convenience!
+	 * Beware of calling this during the shutdown phase, though.  Your module might have been unloaded already.
+	 *
+	 * @return Returns singleton instance, loading the module on demand if needed
+	 */
+	static inline IInterchangeImportModule& Get()
+	{
+		return FModuleManager::LoadModuleChecked< IInterchangeImportModule >( "InterchangeImport" );
+	}
+
+	/**
+	 * Checks to see if this module is loaded and ready.  It is only valid to call Get() if IsAvailable() returns true.
+	 *
+	 * @return True if the module is loaded and ready to use
+	 */
+	static inline bool IsAvailable()
+	{
+		return FModuleManager::Get().IsModuleLoaded( "InterchangeImport" );
+	}
+};
+
diff --git a/Engine/Plugins/Runtime/Interchange/Source/Import/Public/Material/InterchangeMaterialFactory.h b/Engine/Plugins/Runtime/Interchange/Source/Import/Public/Material/InterchangeMaterialFactory.h
--- a/Engine/Plugins/Runtime/Interchange/Source/Import/Public/Material/InterchangeMaterialFactory.h
+++ b/Engine/Plugins/Runtime/Interchange/Source/Import/Public/Material/InterchangeMaterialFactory.h
@@ -10,7 +10,7 @@
 #include "InterchangeMaterialFactory.generated.h"
 
 UCLASS(BlueprintType)
-class INTERCHANGEIMPORTPLUGIN_API UInterchangeMaterialFactory : public UInterchangeFactoryBase
+class INTERCHANGEIMPORT_API UInterchangeMaterialFactory : public UInterchangeFactoryBase
 {
 	GENERATED_BODY()
 public:
diff --git a/Engine/Plugins/Runtime/Interchange/Source/Import/Public/Mesh/InterchangeSkeletalMeshFactory.h b/Engine/Plugins/Runtime/Interchange/Source/Import/Public/Mesh/InterchangeSkeletalMeshFactory.h
--- a/Engine/Plugins/Runtime/Interchange/Source/Import/Public/Mesh/InterchangeSkeletalMeshFactory.h
+++ b/Engine/Plugins/Runtime/Interchange/Source/Import/Public/Mesh/InterchangeSkeletalMeshFactory.h
@@ -10,7 +10,7 @@
 #include "InterchangeSkeletalMeshFactory.generated.h"
 
 UCLASS(BlueprintType)
-class INTERCHANGEIMPORTPLUGIN_API UInterchangeSkeletalMeshFactory : public UInterchangeFactoryBase
+class INTERCHANGEIMPORT_API UInterchangeSkeletalMeshFactory : public UInterchangeFactoryBase
 {
 	GENERATED_BODY()
 public:
diff --git a/Engine/Plugins/Runtime/Interchange/Source/Import/Public/Mesh/InterchangeSkeletalMeshPayloadInterface.h b/Engine/Plugins/Runtime/Interchange/Source/Import/Public/Mesh/InterchangeSkeletalMeshPayloadInterface.h
--- a/Engine/Plugins/Runtime/Interchange/Source/Import/Public/Mesh/InterchangeSkeletalMeshPayloadInterface.h
+++ b/Engine/Plugins/Runtime/Interchange/Source/Import/Public/Mesh/InterchangeSkeletalMeshPayloadInterface.h
@@ -11,7 +11,7 @@
 #include "InterchangeSkeletalMeshPayloadInterface.generated.h"
 
 UINTERFACE()
-class INTERCHANGEIMPORTPLUGIN_API UInterchangeSkeletalMeshPayloadInterface : public UInterface
+class INTERCHANGEIMPORT_API UInterchangeSkeletalMeshPayloadInterface : public UInterface
 {
 	GENERATED_BODY()
 };
@@ -19,7 +19,7 @@ class INTERCHANGEIMPORTPLUGIN_API UInterchangeSkeletalMeshPayloadInterface : pub
 /**
  * Skeletal mesh payload interface. Derive from this interface if your payload can import skeletal mesh
  */
-class INTERCHANGEIMPORTPLUGIN_API IInterchangeSkeletalMeshPayloadInterface
+class INTERCHANGEIMPORT_API IInterchangeSkeletalMeshPayloadInterface
 {
 	GENERATED_BODY()
 public:
diff --git a/Engine/Plugins/Runtime/Interchange/Source/Import/Public/Mesh/InterchangeSkeletonFactory.h b/Engine/Plugins/Runtime/Interchange/Source/Import/Public/Mesh/InterchangeSkeletonFactory.h
--- a/Engine/Plugins/Runtime/Interchange/Source/Import/Public/Mesh/InterchangeSkeletonFactory.h
+++ b/Engine/Plugins/Runtime/Interchange/Source/Import/Public/Mesh/InterchangeSkeletonFactory.h
@@ -10,7 +10,7 @@
 #include "InterchangeSkeletonFactory.generated.h"
 
 UCLASS(BlueprintType)
-class INTERCHANGEIMPORTPLUGIN_API UInterchangeSkeletonFactory : public UInterchangeFactoryBase
+class INTERCHANGEIMPORT_API UInterchangeSkeletonFactory : public UInterchangeFactoryBase
 {
 	GENERATED_BODY()
 public:
diff --git a/Engine/Plugins/Runtime/Interchange/Source/Import/Public/Mesh/InterchangeStaticMeshPayloadInterface.h b/Engine/Plugins/Runtime/Interchange/Source/Import/Public/Mesh/InterchangeStaticMeshPayloadInterface.h
--- a/Engine/Plugins/Runtime/Interchange/Source/Import/Public/Mesh/InterchangeStaticMeshPayloadInterface.h
+++ b/Engine/Plugins/Runtime/Interchange/Source/Import/Public/Mesh/InterchangeStaticMeshPayloadInterface.h
@@ -11,7 +11,7 @@
 #include "InterchangeStaticMeshPayloadInterface.generated.h"
 
 UINTERFACE()
-class INTERCHANGEIMPORTPLUGIN_API UInterchangeStaticMeshPayloadInterface : public UInterface
+class INTERCHANGEIMPORT_API UInterchangeStaticMeshPayloadInterface : public UInterface
 {
 	GENERATED_BODY()
 };
@@ -19,7 +19,7 @@ class INTERCHANGEIMPORTPLUGIN_API UInterchangeStaticMeshPayloadInterface : publi
 /**
  * Static mesh payload interface. Derive from this interface if your payload can import static mesh
  */
-class INTERCHANGEIMPORTPLUGIN_API IInterchangeStaticMeshPayloadInterface
+class INTERCHANGEIMPORT_API IInterchangeStaticMeshPayloadInterface
 {
 	GENERATED_BODY()
 public:
diff --git a/Engine/Plugins/Runtime/Interchange/Source/Import/Public/Texture/InterchangeBMPTranslator.h b/Engine/Plugins/Runtime/Interchange/Source/Import/Public/Texture/InterchangeBMPTranslator.h
--- a/Engine/Plugins/Runtime/Interchange/Source/Import/Public/Texture/InterchangeBMPTranslator.h
+++ b/Engine/Plugins/Runtime/Interchange/Source/Import/Public/Texture/InterchangeBMPTranslator.h
@@ -13,7 +13,7 @@
 #include "InterchangeBMPTranslator.generated.h"
 
 UCLASS(BlueprintType)
-class INTERCHANGEIMPORTPLUGIN_API UInterchangeBMPTranslator : public UInterchangeTranslatorBase, public IInterchangeTexturePayloadInterface
+class INTERCHANGEIMPORT_API UInterchangeBMPTranslator : public UInterchangeTranslatorBase, public IInterchangeTexturePayloadInterface
 {
 	GENERATED_BODY()
 public:
diff --git a/Engine/Plugins/Runtime/Interchange/Source/Import/Public/Texture/InterchangeBlockedTexturePayloadInterface.h b/Engine/Plugins/Runtime/Interchange/Source/Import/Public/Texture/InterchangeBlockedTexturePayloadInterface.h
--- a/Engine/Plugins/Runtime/Interchange/Source/Import/Public/Texture/InterchangeBlockedTexturePayloadInterface.h
+++ b/Engine/Plugins/Runtime/Interchange/Source/Import/Public/Texture/InterchangeBlockedTexturePayloadInterface.h
@@ -14,7 +14,7 @@ class UInterchangeSourceData;
 class FString;
 
 UINTERFACE()
-class INTERCHANGEIMPORTPLUGIN_API UInterchangeBlockedTexturePayloadInterface : public UInterface
+class INTERCHANGEIMPORT_API UInterchangeBlockedTexturePayloadInterface : public UInterface
 {
 	GENERATED_BODY()
 };
@@ -22,12 +22,12 @@ class INTERCHANGEIMPORTPLUGIN_API UInterchangeBlockedTexturePayloadInterface : p
 /**
  * Blocked Texture payload interface (Also know as UDIM(s)). Derive from it if your payload can import blocked/UDIMs texture
  */
-class INTERCHANGEIMPORTPLUGIN_API IInterchangeBlockedTexturePayloadInterface
+class INTERCHANGEIMPORT_API IInterchangeBlockedTexturePayloadInterface
 {
 	GENERATED_BODY()
 public:
 
-	
+
 	/**
 	 * Once the translation is done, the import process need a way to retrieve payload data.
 	 * This payload will be use by the factories to create the asset.
diff --git a/Engine/Plugins/Runtime/Interchange/Source/Import/Public/Texture/InterchangeDDSTranslator.h b/Engine/Plugins/Runtime/Interchange/Source/Import/Public/Texture/InterchangeDDSTranslator.h
--- a/Engine/Plugins/Runtime/Interchange/Source/Import/Public/Texture/InterchangeDDSTranslator.h
+++ b/Engine/Plugins/Runtime/Interchange/Source/Import/Public/Texture/InterchangeDDSTranslator.h
@@ -13,7 +13,7 @@
 #include "InterchangeDDSTranslator.generated.h"
 
 UCLASS(BlueprintType)
-class INTERCHANGEIMPORTPLUGIN_API UInterchangeDDSTranslator : public UInterchangeTranslatorBase, public IInterchangeTexturePayloadInterface
+class INTERCHANGEIMPORT_API UInterchangeDDSTranslator : public UInterchangeTranslatorBase, public IInterchangeTexturePayloadInterface
 {
 	GENERATED_BODY()
 public:
diff --git a/Engine/Plugins/Runtime/Interchange/Source/Import/Public/Texture/InterchangeEXRTranslator.h b/Engine/Plugins/Runtime/Interchange/Source/Import/Public/Texture/InterchangeEXRTranslator.h
--- a/Engine/Plugins/Runtime/Interchange/Source/Import/Public/Texture/InterchangeEXRTranslator.h
+++ b/Engine/Plugins/Runtime/Interchange/Source/Import/Public/Texture/InterchangeEXRTranslator.h
@@ -13,7 +13,7 @@
 #include "InterchangeEXRTranslator.generated.h"
 
 UCLASS(BlueprintType)
-class INTERCHANGEIMPORTPLUGIN_API UInterchangeEXRTranslator : public UInterchangeTranslatorBase, public IInterchangeTexturePayloadInterface
+class INTERCHANGEIMPORT_API UInterchangeEXRTranslator : public UInterchangeTranslatorBase, public IInterchangeTexturePayloadInterface
 {
 	GENERATED_BODY()
 public:
diff --git a/Engine/Plugins/Runtime/Interchange/Source/Import/Public/Texture/InterchangeJPGTranslator.h b/Engine/Plugins/Runtime/Interchange/Source/Import/Public/Texture/InterchangeJPGTranslator.h
--- a/Engine/Plugins/Runtime/Interchange/Source/Import/Public/Texture/InterchangeJPGTranslator.h
+++ b/Engine/Plugins/Runtime/Interchange/Source/Import/Public/Texture/InterchangeJPGTranslator.h
@@ -13,7 +13,7 @@
 #include "InterchangeJPGTranslator.generated.h"
 
 UCLASS(BlueprintType)
-class INTERCHANGEIMPORTPLUGIN_API UInterchangeJPGTranslator : public UInterchangeTranslatorBase, public IInterchangeTexturePayloadInterface
+class INTERCHANGEIMPORT_API UInterchangeJPGTranslator : public UInterchangeTranslatorBase, public IInterchangeTexturePayloadInterface
 {
 	GENERATED_BODY()
 public:
diff --git a/Engine/Plugins/Runtime/Interchange/Source/Import/Public/Texture/InterchangePCXTranslator.h b/Engine/Plugins/Runtime/Interchange/Source/Import/Public/Texture/InterchangePCXTranslator.h
--- a/Engine/Plugins/Runtime/Interchange/Source/Import/Public/Texture/InterchangePCXTranslator.h
+++ b/Engine/Plugins/Runtime/Interchange/Source/Import/Public/Texture/InterchangePCXTranslator.h
@@ -13,7 +13,7 @@
 #include "InterchangePCXTranslator.generated.h"
 
 UCLASS(BlueprintType)
-class INTERCHANGEIMPORTPLUGIN_API UInterchangePCXTranslator : public UInterchangeTranslatorBase, public IInterchangeTexturePayloadInterface
+class INTERCHANGEIMPORT_API UInterchangePCXTranslator : public UInterchangeTranslatorBase, public IInterchangeTexturePayloadInterface
 {
 	GENERATED_BODY()
 public:
diff --git a/Engine/Plugins/Runtime/Interchange/Source/Import/Public/Texture/InterchangePNGTranslator.h b/Engine/Plugins/Runtime/Interchange/Source/Import/Public/Texture/InterchangePNGTranslator.h
--- a/Engine/Plugins/Runtime/Interchange/Source/Import/Public/Texture/InterchangePNGTranslator.h
+++ b/Engine/Plugins/Runtime/Interchange/Source/Import/Public/Texture/InterchangePNGTranslator.h
@@ -13,7 +13,7 @@
 #include "InterchangePNGTranslator.generated.h"
 
 UCLASS(BlueprintType)
-class INTERCHANGEIMPORTPLUGIN_API UInterchangePNGTranslator : public UInterchangeTranslatorBase, public IInterchangeTexturePayloadInterface
+class INTERCHANGEIMPORT_API UInterchangePNGTranslator : public UInterchangeTranslatorBase, public IInterchangeTexturePayloadInterface
 {
 	GENERATED_BODY()
 public:
diff --git a/Engine/Plugins/Runtime/Interchange/Source/Import/Public/Texture/InterchangePSDTranslator.h b/Engine/Plugins/Runtime/Interchange/Source/Import/Public/Texture/InterchangePSDTranslator.h
--- a/Engine/Plugins/Runtime/Interchange/Source/Import/Public/Texture/InterchangePSDTranslator.h
+++ b/Engine/Plugins/Runtime/Interchange/Source/Import/Public/Texture/InterchangePSDTranslator.h
@@ -13,7 +13,7 @@
 #include "InterchangePSDTranslator.generated.h"
 
 UCLASS(BlueprintType)
-class INTERCHANGEIMPORTPLUGIN_API UInterchangePSDTranslator : public UInterchangeTranslatorBase, public IInterchangeTexturePayloadInterface
+class INTERCHANGEIMPORT_API UInterchangePSDTranslator : public UInterchangeTranslatorBase, public IInterchangeTexturePayloadInterface
 {
 	GENERATED_BODY()
 public:
diff --git a/Engine/Plugins/Runtime/Interchange/Source/Import/Public/Texture/InterchangeTGATranslator.h b/Engine/Plugins/Runtime/Interchange/Source/Import/Public/Texture/InterchangeTGATranslator.h
--- a/Engine/Plugins/Runtime/Interchange/Source/Import/Public/Texture/InterchangeTGATranslator.h
+++ b/Engine/Plugins/Runtime/Interchange/Source/Import/Public/Texture/InterchangeTGATranslator.h
@@ -12,7 +12,7 @@
 #include "InterchangeTGATranslator.generated.h"
 
 UCLASS(BlueprintType)
-class INTERCHANGEIMPORTPLUGIN_API UInterchangeTGATranslator : public UInterchangeTranslatorBase, public IInterchangeTexturePayloadInterface
+class INTERCHANGEIMPORT_API UInterchangeTGATranslator : public UInterchangeTranslatorBase, public IInterchangeTexturePayloadInterface
 {
 	GENERATED_BODY()
 public:
diff --git a/Engine/Plugins/Runtime/Interchange/Source/Import/Public/Texture/InterchangeTextureFactory.h b/Engine/Plugins/Runtime/Interchange/Source/Import/Public/Texture/InterchangeTextureFactory.h
--- a/Engine/Plugins/Runtime/Interchange/Source/Import/Public/Texture/InterchangeTextureFactory.h
+++ b/Engine/Plugins/Runtime/Interchange/Source/Import/Public/Texture/InterchangeTextureFactory.h
@@ -11,7 +11,7 @@
 #include "InterchangeTextureFactory.generated.h"
 
 UCLASS(BlueprintType)
-class INTERCHANGEIMPORTPLUGIN_API UInterchangeTextureFactory : public UInterchangeFactoryBase
+class INTERCHANGEIMPORT_API UInterchangeTextureFactory : public UInterchangeFactoryBase
 {
 	GENERATED_BODY()
 public:
diff --git a/Engine/Plugins/Runtime/Interchange/Source/Import/Public/Texture/InterchangeTexturePayloadData.h b/Engine/Plugins/Runtime/Interchange/Source/Import/Public/Texture/InterchangeTexturePayloadData.h
--- a/Engine/Plugins/Runtime/Interchange/Source/Import/Public/Texture/InterchangeTexturePayloadData.h
+++ b/Engine/Plugins/Runtime/Interchange/Source/Import/Public/Texture/InterchangeTexturePayloadData.h
@@ -9,7 +9,7 @@ namespace UE
 {
 	namespace Interchange
 	{
-		struct INTERCHANGEIMPORTPLUGIN_API FImportImage
+		struct INTERCHANGEIMPORT_API FImportImage
 		{
 			TArray64<uint8> RawData;
 			ETextureSourceFormat Format = TSF_Invalid;
@@ -28,7 +28,7 @@ namespace UE
 			void* GetMipData(int32 InMipIndex);
 		};
 
-		struct INTERCHANGEIMPORTPLUGIN_API FImportImageHelper
+		struct INTERCHANGEIMPORT_API FImportImageHelper
 		{
 			/**
 			 * Tests if the given height and width specify a supported texture resolution to import; Can optionally check if the height/width are powers of two
diff --git a/Engine/Plugins/Runtime/Interchange/Source/Import/Public/Texture/InterchangeTexturePayloadInterface.h b/Engine/Plugins/Runtime/Interchange/Source/Import/Public/Texture/InterchangeTexturePayloadInterface.h
--- a/Engine/Plugins/Runtime/Interchange/Source/Import/Public/Texture/InterchangeTexturePayloadInterface.h
+++ b/Engine/Plugins/Runtime/Interchange/Source/Import/Public/Texture/InterchangeTexturePayloadInterface.h
@@ -10,7 +10,7 @@
 #include "InterchangeTexturePayloadInterface.generated.h"
 
 UINTERFACE()
-class INTERCHANGEIMPORTPLUGIN_API UInterchangeTexturePayloadInterface : public UInterface
+class INTERCHANGEIMPORT_API UInterchangeTexturePayloadInterface : public UInterface
 {
 	GENERATED_BODY()
 };
@@ -18,7 +18,7 @@ class INTERCHANGEIMPORTPLUGIN_API UInterchangeTexturePayloadInterface : public U
 /**
  * Texture payload interface. Derive from it if your payload can import texture
  */
-class INTERCHANGEIMPORTPLUGIN_API IInterchangeTexturePayloadInterface
+class INTERCHANGEIMPORT_API IInterchangeTexturePayloadInterface
 {
 	GENERATED_BODY()
 public:
diff --git a/Engine/Plugins/Runtime/Interchange/Source/Import/Public/Texture/InterchangeUDIMTranslator.h b/Engine/Plugins/Runtime/Interchange/Source/Import/Public/Texture/InterchangeUDIMTranslator.h
--- a/Engine/Plugins/Runtime/Interchange/Source/Import/Public/Texture/InterchangeUDIMTranslator.h
+++ b/Engine/Plugins/Runtime/Interchange/Source/Import/Public/Texture/InterchangeUDIMTranslator.h
@@ -16,7 +16,7 @@ class UInterchangeBaseNodeContainer;
 class UInterchangeSourceData;
 
 UCLASS(BlueprintType)
-class INTERCHANGEIMPORTPLUGIN_API UInterchangeUDIMTranslator : public UInterchangeTranslatorBase, public IInterchangeBlockedTexturePayloadInterface
+class INTERCHANGEIMPORT_API UInterchangeUDIMTranslator : public UInterchangeTranslatorBase, public IInterchangeBlockedTexturePayloadInterface
 {
 	GENERATED_BODY()
 
diff --git a/Engine/Plugins/Runtime/Interchange/Source/Nodes/InterchangeNodePlugin.Build.cs b/Engine/Plugins/Runtime/Interchange/Source/Nodes/InterchangeNodePlugin.Build.cs
deleted file mode 100644
diff --git a/Engine/Plugins/Runtime/Interchange/Source/Nodes/InterchangeNodes.Build.cs b/Engine/Plugins/Runtime/Interchange/Source/Nodes/InterchangeNodes.Build.cs
new file mode 100644
--- /dev/null
+++ b/Engine/Plugins/Runtime/Interchange/Source/Nodes/InterchangeNodes.Build.cs
@@ -0,0 +1,24 @@
+// Copyright Epic Games, Inc. All Rights Reserved.
+
+namespace UnrealBuildTool.Rules
+{
+	public class InterchangeNodes : ModuleRules
+	{
+		public InterchangeNodes(ReadOnlyTargetRules Target) : base(Target)
+		{
+			PublicDependencyModuleNames.AddRange(
+				new string[]
+				{
+					"Core",
+					"CoreUObject",
+					"InterchangeCore",
+				}
+			);
+
+			if(Target.bCompileAgainstEngine)
+			{
+				PublicDependencyModuleNames.Add("Engine");
+			}
+		}
+	}
+}
diff --git a/Engine/Plugins/Runtime/Interchange/Source/Nodes/Private/InterchangeNodePlugin.cpp b/Engine/Plugins/Runtime/Interchange/Source/Nodes/Private/InterchangeNodePlugin.cpp
deleted file mode 100644
diff --git a/Engine/Plugins/Runtime/Interchange/Source/Nodes/Private/InterchangeNodesModule.cpp b/Engine/Plugins/Runtime/Interchange/Source/Nodes/Private/InterchangeNodesModule.cpp
new file mode 100644
--- /dev/null
+++ b/Engine/Plugins/Runtime/Interchange/Source/Nodes/Private/InterchangeNodesModule.cpp
@@ -0,0 +1,36 @@
+// Copyright Epic Games, Inc. All Rights Reserved.
+
+#include "InterchangeNodesModule.h"
+
+#include "CoreMinimal.h"
+#include "Modules/ModuleManager.h"
+//#include "InterchangeManager.h"
+
+DEFINE_LOG_CATEGORY(LogInterchangeNodes);
+
+class FInterchangeNodesModule : public IInterchangeNodesModule
+{
+	/** IModuleInterface implementation */
+	virtual void StartupModule() override;
+	virtual void ShutdownModule() override;
+};
+
+IMPLEMENT_MODULE(FInterchangeNodesModule, InterchangeNodes)
+
+
+
+void FInterchangeNodesModule::StartupModule()
+{
+	//Register anything needed to the interchange manager
+	//UInterchangeManager& InterchangeManager = UInterchangeManager::GetInterchangeManager();
+	//InterchangeManager.Register...
+}
+
+
+void FInterchangeNodesModule::ShutdownModule()
+{
+
+}
+
+
+
diff --git a/Engine/Plugins/Runtime/Interchange/Source/Nodes/Public/IInterchangeNodePlugin.h b/Engine/Plugins/Runtime/Interchange/Source/Nodes/Public/IInterchangeNodePlugin.h
deleted file mode 100644
diff --git a/Engine/Plugins/Runtime/Interchange/Source/Nodes/Public/InterchangeJointNode.h b/Engine/Plugins/Runtime/Interchange/Source/Nodes/Public/InterchangeJointNode.h
--- a/Engine/Plugins/Runtime/Interchange/Source/Nodes/Public/InterchangeJointNode.h
+++ b/Engine/Plugins/Runtime/Interchange/Source/Nodes/Public/InterchangeJointNode.h
@@ -13,7 +13,7 @@
 
 
 UCLASS(BlueprintType)
-class INTERCHANGENODEPLUGIN_API UInterchangeJointNode : public UInterchangeBaseNode
+class INTERCHANGENODES_API UInterchangeJointNode : public UInterchangeBaseNode
 {
 	GENERATED_BODY()
 
diff --git a/Engine/Plugins/Runtime/Interchange/Source/Nodes/Public/InterchangeMaterialNode.h b/Engine/Plugins/Runtime/Interchange/Source/Nodes/Public/InterchangeMaterialNode.h
--- a/Engine/Plugins/Runtime/Interchange/Source/Nodes/Public/InterchangeMaterialNode.h
+++ b/Engine/Plugins/Runtime/Interchange/Source/Nodes/Public/InterchangeMaterialNode.h
@@ -70,7 +70,7 @@ enum class EInterchangeMaterialNodeParameterName : uint8
 };
 
 UCLASS(BlueprintType)
-class INTERCHANGENODEPLUGIN_API UInterchangeMaterialNode : public UInterchangeBaseNode
+class INTERCHANGENODES_API UInterchangeMaterialNode : public UInterchangeBaseNode
 {
 	GENERATED_BODY()
 
diff --git a/Engine/Plugins/Runtime/Interchange/Source/Nodes/Public/InterchangeNodesModule.h b/Engine/Plugins/Runtime/Interchange/Source/Nodes/Public/InterchangeNodesModule.h
new file mode 100644
--- /dev/null
+++ b/Engine/Plugins/Runtime/Interchange/Source/Nodes/Public/InterchangeNodesModule.h
@@ -0,0 +1,40 @@
+// Copyright Epic Games, Inc. All Rights Reserved.
+
+#pragma once
+
+#include "CoreMinimal.h"
+#include "Modules/ModuleInterface.h"
+#include "Modules/ModuleManager.h"
+
+DECLARE_LOG_CATEGORY_EXTERN(LogInterchangeNodes, Log, All);
+
+/**
+ * The public interface to this module
+ */
+class IInterchangeNodesModule : public IModuleInterface
+{
+
+public:
+
+	/**
+	 * Singleton-like access to this module's interface.  This is just for convenience!
+	 * Beware of calling this during the shutdown phase, though.  Your module might have been unloaded already.
+	 *
+	 * @return Returns singleton instance, loading the module on demand if needed
+	 */
+	static inline IInterchangeNodesModule& Get()
+	{
+		return FModuleManager::LoadModuleChecked< IInterchangeNodesModule >( "InterchangeNodes" );
+	}
+
+	/**
+	 * Checks to see if this module is loaded and ready.  It is only valid to call Get() if IsAvailable() returns true.
+	 *
+	 * @return True if the module is loaded and ready to use
+	 */
+	static inline bool IsAvailable()
+	{
+		return FModuleManager::Get().IsModuleLoaded( "InterchangeNodes" );
+	}
+};
+
diff --git a/Engine/Plugins/Runtime/Interchange/Source/Nodes/Public/InterchangeSceneNode.h b/Engine/Plugins/Runtime/Interchange/Source/Nodes/Public/InterchangeSceneNode.h
--- a/Engine/Plugins/Runtime/Interchange/Source/Nodes/Public/InterchangeSceneNode.h
+++ b/Engine/Plugins/Runtime/Interchange/Source/Nodes/Public/InterchangeSceneNode.h
@@ -9,7 +9,7 @@
 
 
 UCLASS(BlueprintType)
-class INTERCHANGENODEPLUGIN_API UInterchangeSceneNode : public UInterchangeBaseNode
+class INTERCHANGENODES_API UInterchangeSceneNode : public UInterchangeBaseNode
 {
 	GENERATED_BODY()
 
diff --git a/Engine/Plugins/Runtime/Interchange/Source/Nodes/Public/InterchangeSkeletalMeshLodDataNode.h b/Engine/Plugins/Runtime/Interchange/Source/Nodes/Public/InterchangeSkeletalMeshLodDataNode.h
--- a/Engine/Plugins/Runtime/Interchange/Source/Nodes/Public/InterchangeSkeletalMeshLodDataNode.h
+++ b/Engine/Plugins/Runtime/Interchange/Source/Nodes/Public/InterchangeSkeletalMeshLodDataNode.h
@@ -32,7 +32,7 @@ namespace UE
 }//ns UE
 
 UCLASS(BlueprintType)
-class INTERCHANGENODEPLUGIN_API UInterchangeSkeletalMeshLodDataNode : public UInterchangeBaseNode
+class INTERCHANGENODES_API UInterchangeSkeletalMeshLodDataNode : public UInterchangeBaseNode
 {
 	GENERATED_BODY()
 
diff --git a/Engine/Plugins/Runtime/Interchange/Source/Nodes/Public/InterchangeSkeletalMeshNode.h b/Engine/Plugins/Runtime/Interchange/Source/Nodes/Public/InterchangeSkeletalMeshNode.h
--- a/Engine/Plugins/Runtime/Interchange/Source/Nodes/Public/InterchangeSkeletalMeshNode.h
+++ b/Engine/Plugins/Runtime/Interchange/Source/Nodes/Public/InterchangeSkeletalMeshNode.h
@@ -30,7 +30,7 @@ namespace UE
 }//ns UE
 
 UCLASS(BlueprintType)
-class INTERCHANGENODEPLUGIN_API UInterchangeSkeletalMeshNode : public UInterchangeBaseNode
+class INTERCHANGENODES_API UInterchangeSkeletalMeshNode : public UInterchangeBaseNode
 {
 	GENERATED_BODY()
 
diff --git a/Engine/Plugins/Runtime/Interchange/Source/Nodes/Public/InterchangeSkeletonNode.h b/Engine/Plugins/Runtime/Interchange/Source/Nodes/Public/InterchangeSkeletonNode.h
--- a/Engine/Plugins/Runtime/Interchange/Source/Nodes/Public/InterchangeSkeletonNode.h
+++ b/Engine/Plugins/Runtime/Interchange/Source/Nodes/Public/InterchangeSkeletonNode.h
@@ -12,7 +12,7 @@
 #include "InterchangeSkeletonNode.generated.h"
 
 UCLASS(BlueprintType)
-class INTERCHANGENODEPLUGIN_API UInterchangeSkeletonNode : public UInterchangeBaseNode
+class INTERCHANGENODES_API UInterchangeSkeletonNode : public UInterchangeBaseNode
 {
 	GENERATED_BODY()
 
diff --git a/Engine/Plugins/Runtime/Interchange/Source/Nodes/Public/InterchangeTextureNode.h b/Engine/Plugins/Runtime/Interchange/Source/Nodes/Public/InterchangeTextureNode.h
--- a/Engine/Plugins/Runtime/Interchange/Source/Nodes/Public/InterchangeTextureNode.h
+++ b/Engine/Plugins/Runtime/Interchange/Source/Nodes/Public/InterchangeTextureNode.h
@@ -10,6 +10,7 @@
 #include "Engine/Texture2D.h"
 #endif
 
+
 #include "InterchangeTextureNode.generated.h"
 
 //Interchange namespace
@@ -31,7 +32,7 @@ namespace UE
 }//ns UE
 
 UCLASS(BlueprintType)
-class INTERCHANGENODEPLUGIN_API UInterchangeTextureNode : public UInterchangeBaseNode
+class INTERCHANGENODES_API UInterchangeTextureNode : public UInterchangeBaseNode
 {
 	GENERATED_BODY()
 
@@ -584,7 +585,7 @@ public:
 	{
 		IMPLEMENT_NODE_ATTRIBUTE_SETTER(UInterchangeTextureNode, Filter, uint8, UTexture)
 	}
-	
+
 	/** Return false if the Attribute was not set previously.*/
 	UFUNCTION(BlueprintCallable, Category = "Interchange | Node | Texture")
 	bool GetCustomMipLoadOptions(uint8& AttributeValue) const
@@ -597,7 +598,7 @@ public:
 	{
 		IMPLEMENT_NODE_ATTRIBUTE_SETTER(UInterchangeTextureNode, MipLoadOptions, uint8, UTexture)
 	}
-	
+
 	/** Return false if the Attribute was not set previously.*/
 	UFUNCTION(BlueprintCallable, Category = "Interchange | Node | Texture")
 	bool GetCustomSRGB(bool& AttributeValue) const
@@ -686,7 +687,7 @@ public:
 		IMPLEMENT_NODE_ATTRIBUTE_SETTER_NODELEGATE(bPreserveBorder, bool)
 #endif
 	}
-	
+
 	/** Return false if the Attribute was not set previously.*/
 	UFUNCTION(BlueprintCallable, Category = "Interchange | Node | Texture")
 	bool GetCustomMipGenSettings(uint8& AttributeValue) const
@@ -759,7 +760,7 @@ public:
 
 	//////////////////////////////////////////////////////////////////////////
 	//Compositing
-	
+
 	/** Return false if the Attribute was not set previously.*/
 	UFUNCTION(BlueprintCallable, Category = "Interchange | Node | Texture")
 	bool GetCustomCompositeTextureMode(uint8& AttributeValue) const
@@ -776,7 +777,7 @@ public:
 		IMPLEMENT_NODE_ATTRIBUTE_SETTER_NODELEGATE(CompositeTextureMode, uint8)
 #endif
 	}
-	
+
 	/** Return false if the Attribute was not set previously.*/
 	UFUNCTION(BlueprintCallable, Category = "Interchange | Node | Texture")
 	bool GetCustomCompositePower(float& AttributeValue) const
@@ -930,7 +931,7 @@ private:
 #endif //#if WITH_EDITORONLY_DATA
 	IMPLEMENT_NODE_ATTRIBUTE_APPLY_UOBJECT(LODBias, int32, UTexture, );
 	IMPLEMENT_NODE_ATTRIBUTE_APPLY_UOBJECT(LODGroup, uint8, UTexture, TEnumAsByte<enum TextureGroup>);
-	
+
 	//TODO support per platform data in the FAttributeStorage, so we can set different value per platform at the pipeline stage, We set only the default value for now
 	//IMPLEMENT_NODE_ATTRIBUTE_APPLY_UOBJECT(Downscale, float, UTexture, );
 #if WITH_ENGINE
diff --git a/Engine/Plugins/Runtime/Interchange/Source/Parsers/Fbx/InterchangeFbxParser.Build.cs b/Engine/Plugins/Runtime/Interchange/Source/Parsers/Fbx/InterchangeFbxParser.Build.cs
--- a/Engine/Plugins/Runtime/Interchange/Source/Parsers/Fbx/InterchangeFbxParser.Build.cs
+++ b/Engine/Plugins/Runtime/Interchange/Source/Parsers/Fbx/InterchangeFbxParser.Build.cs
@@ -9,11 +9,11 @@ namespace UnrealBuildTool.Rules
 		public InterchangeFbxParser(ReadOnlyTargetRules Target) : base(Target)
 		{
 			bEnableExceptions = true;
-			
-            PublicDependencyModuleNames.AddRange(
+
+			PublicDependencyModuleNames.AddRange(
 				new string[]
 				{
-                    "Core",
+					"Core",
 					"CoreUObject",
 					"InterchangeCore",
 				}
@@ -22,7 +22,7 @@ namespace UnrealBuildTool.Rules
 			PrivateDependencyModuleNames.AddRange(
 				new string[]
 				{
-					"InterchangeNodePlugin",
+					"InterchangeNodes",
 					"MeshDescription",
 					"SkeletalMeshDescription",
 					"StaticMeshDescription"
@@ -33,7 +33,7 @@ namespace UnrealBuildTool.Rules
 			{
 				PrivateDependencyModuleNames.AddRange(
 					new string[] {
-					"Engine",
+						"Engine",
 					}
 				);
 			}
@@ -42,5 +42,5 @@ namespace UnrealBuildTool.Rules
 				"FBX"
 			);
 		}
-    }
+	}
 }
diff --git a/Engine/Plugins/Runtime/Interchange/Source/Pipelines/InterchangePipelinePlugin.Build.cs b/Engine/Plugins/Runtime/Interchange/Source/Pipelines/InterchangePipelinePlugin.Build.cs
deleted file mode 100644
diff --git a/Engine/Plugins/Runtime/Interchange/Source/Pipelines/InterchangePipelines.Build.cs b/Engine/Plugins/Runtime/Interchange/Source/Pipelines/InterchangePipelines.Build.cs
new file mode 100644
--- /dev/null
+++ b/Engine/Plugins/Runtime/Interchange/Source/Pipelines/InterchangePipelines.Build.cs
@@ -0,0 +1,22 @@
+// Copyright Epic Games, Inc. All Rights Reserved.
+
+namespace UnrealBuildTool.Rules
+{
+	public class InterchangePipelines : ModuleRules
+	{
+		public InterchangePipelines(ReadOnlyTargetRules Target) : base(Target)
+		{
+			PublicDependencyModuleNames.AddRange(
+				new string[]
+				{
+					"Core",
+					"CoreUObject",
+					"Engine",
+					"InterchangeCore",
+					"InterchangeEngine",
+					"InterchangeNodes"
+				}
+			);
+		}
+	}
+}
diff --git a/Engine/Plugins/Runtime/Interchange/Source/Pipelines/Private/InterchangePipelinePlugin.cpp b/Engine/Plugins/Runtime/Interchange/Source/Pipelines/Private/InterchangePipelinePlugin.cpp
deleted file mode 100644
diff --git a/Engine/Plugins/Runtime/Interchange/Source/Pipelines/Private/InterchangePipelinesModule.cpp b/Engine/Plugins/Runtime/Interchange/Source/Pipelines/Private/InterchangePipelinesModule.cpp
new file mode 100644
--- /dev/null
+++ b/Engine/Plugins/Runtime/Interchange/Source/Pipelines/Private/InterchangePipelinesModule.cpp
@@ -0,0 +1,50 @@
+// Copyright Epic Games, Inc. All Rights Reserved.
+
+#include "InterchangePipelinesModule.h"
+
+#include "CoreMinimal.h"
+#include "Engine/Engine.h"
+#include "InterchangeGenericMeshPipeline.h"
+#include "InterchangeManager.h"
+#include "Misc/CoreDelegates.h"
+#include "Modules/ModuleManager.h"
+
+
+class FInterchangePipelinesModule : public IInterchangePipelinesModule
+{
+	/** IModuleInterface implementation */
+	virtual void StartupModule() override;
+	virtual void ShutdownModule() override;
+};
+
+IMPLEMENT_MODULE(FInterchangePipelinesModule, InterchangePipelines)
+
+
+
+void FInterchangePipelinesModule::StartupModule()
+{
+	auto RegisterItems = []()
+	{
+		UInterchangeManager& InterchangeManager = UInterchangeManager::GetInterchangeManager();
+		//Register the Pipeline
+		//InterchangeManager.RegisterPipeline(UInterchangeGenericMeshPipeline::StaticClass());
+	};
+
+	if (GEngine)
+	{
+		RegisterItems();
+	}
+	else
+	{
+		FCoreDelegates::OnPostEngineInit.AddLambda(RegisterItems);
+	}
+}
+
+
+void FInterchangePipelinesModule::ShutdownModule()
+{
+
+}
+
+
+
diff --git a/Engine/Plugins/Runtime/Interchange/Source/Pipelines/Public/IInterchangePipelinePlugin.h b/Engine/Plugins/Runtime/Interchange/Source/Pipelines/Public/IInterchangePipelinePlugin.h
deleted file mode 100644
diff --git a/Engine/Plugins/Runtime/Interchange/Source/Pipelines/Public/InterchangeGenericMeshPipeline.h b/Engine/Plugins/Runtime/Interchange/Source/Pipelines/Public/InterchangeGenericMeshPipeline.h
--- a/Engine/Plugins/Runtime/Interchange/Source/Pipelines/Public/InterchangeGenericMeshPipeline.h
+++ b/Engine/Plugins/Runtime/Interchange/Source/Pipelines/Public/InterchangeGenericMeshPipeline.h
@@ -23,7 +23,7 @@ class UInterchangeSkeletalMeshNode;
  * @meta MeshType - String, the property is for static or skeletal or both (static | skeletal) mesh type. If not specified it will apply to all mesh type.
  */
 UCLASS(BlueprintType)
-class INTERCHANGEPIPELINEPLUGIN_API UInterchangeGenericMeshPipeline : public UInterchangePipelineBase
+class INTERCHANGEPIPELINES_API UInterchangeGenericMeshPipeline : public UInterchangePipelineBase
 {
 	GENERATED_BODY()
 
@@ -39,7 +39,7 @@ public:
 	//////////////////////////////////////////////////////////////////////////
 
 protected:
-	
+
 	virtual bool ExecutePreImportPipeline(UInterchangeBaseNodeContainer* BaseNodeContainer, const TArray<UInterchangeSourceData*>& SourceDatas) override;
 
 	//virtual bool ExecutePostImportPipeline(const UInterchangeBaseNodeContainer* BaseNodeContainer, const FName& NodeKey, UObject* CreatedAsset) override;
diff --git a/Engine/Plugins/Runtime/Interchange/Source/Pipelines/Public/InterchangePipelinesModule.h b/Engine/Plugins/Runtime/Interchange/Source/Pipelines/Public/InterchangePipelinesModule.h
new file mode 100644
--- /dev/null
+++ b/Engine/Plugins/Runtime/Interchange/Source/Pipelines/Public/InterchangePipelinesModule.h
@@ -0,0 +1,39 @@
+// Copyright Epic Games, Inc. All Rights Reserved.
+
+#pragma once
+
+#include "CoreMinimal.h"
+#include "Modules/ModuleInterface.h"
+#include "Modules/ModuleManager.h"
+
+
+/**
+ * The public interface to this module
+ */
+class IInterchangePipelinesModule : public IModuleInterface
+{
+
+public:
+
+	/**
+	 * Singleton-like access to this module's interface.  This is just for convenience!
+	 * Beware of calling this during the shutdown phase, though.  Your module might have been unloaded already.
+	 *
+	 * @return Returns singleton instance, loading the module on demand if needed
+	 */
+	static inline IInterchangePipelinesModule& Get()
+	{
+		return FModuleManager::LoadModuleChecked< IInterchangePipelinesModule >( "InterchangePipelines" );
+	}
+
+	/**
+	 * Checks to see if this module is loaded and ready.  It is only valid to call Get() if IsAvailable() returns true.
+	 *
+	 * @return True if the module is loaded and ready to use
+	 */
+	static inline bool IsAvailable()
+	{
+		return FModuleManager::Get().IsModuleLoaded( "InterchangePipelines" );
+	}
+};
+
diff --git a/Engine/Source/Programs/InterchangeWorker/InterchangeWorker.Build.cs b/Engine/Source/Programs/InterchangeWorker/InterchangeWorker.Build.cs
--- a/Engine/Source/Programs/InterchangeWorker/InterchangeWorker.Build.cs
+++ b/Engine/Source/Programs/InterchangeWorker/InterchangeWorker.Build.cs
@@ -33,7 +33,7 @@ public class InterchangeWorker : ModuleRules
 				{
 					"InterchangeCore",
 					"InterchangeFbxParser",
-					"InterchangeNodePlugin"
+					"InterchangeNodes"
 				}
 			);
 
