Fixed a crash caused by the change of the object ptr properties to the TObjectPtr.
Changed the function UpdateImportAssetData to let the user manually set the source data on his asset.

diff --git a/Engine/Plugins/Runtime/Interchange/Source/Import/Private/InterchangeImportCommon.cpp b/Engine/Plugins/Runtime/Interchange/Source/Import/Private/InterchangeImportCommon.cpp
--- a/Engine/Plugins/Runtime/Interchange/Source/Import/Private/InterchangeImportCommon.cpp
+++ b/Engine/Plugins/Runtime/Interchange/Source/Import/Private/InterchangeImportCommon.cpp
@@ -13,7 +13,7 @@ namespace UE
 	namespace Interchange
 	{
 		FFactoryCommon::FUpdateImportAssetDataParameters::FUpdateImportAssetDataParameters(UObject* InAssetImportDataOuter
-																							, UAssetImportData** InAssetImportData
+																							, UAssetImportData* InAssetImportData
 																							, const UInterchangeSourceData* InSourceData
 																							, FString InNodeUniqueID
 																							, UInterchangeBaseNodeContainer* InNodeContainer)
@@ -30,42 +30,32 @@ namespace UE
 			ensure(NodeContainer);
 		}
 
-		FFactoryCommon::FUpdateImportAssetDataParameters::FUpdateImportAssetDataParameters(UObject* InAssetImportDataOuter
-																							, TObjectPtr<UAssetImportData>* InAssetImportData
-																							, const UInterchangeSourceData* InSourceData
-																							, FString InNodeUniqueID
-																							, UInterchangeBaseNodeContainer* InNodeContainer)
-			: FUpdateImportAssetDataParameters(InAssetImportDataOuter, &((UAssetImportData*&)InAssetImportData), InSourceData, InNodeUniqueID, InNodeContainer)
-		{
-		}
-
-		void FFactoryCommon::UpdateImportAssetData(FUpdateImportAssetDataParameters& Parameters)
+		UAssetImportData* FFactoryCommon::UpdateImportAssetData(FUpdateImportAssetDataParameters& Parameters)
 		{
 #if WITH_EDITORONLY_DATA
 			if (!ensure(IsInGameThread()))
 			{
-				return;
+				return nullptr;
 			}
-			if (!ensure(Parameters.SourceData && Parameters.AssetImportDataOuter && Parameters.AssetImportData))
+			if (!ensure(Parameters.SourceData && Parameters.AssetImportDataOuter))
 			{
-				return;
+				return nullptr;
 			}
 			//Set the asset import data file source to allow reimport. TODO: manage MD5 Hash properly
 			TOptional<FMD5Hash> FileContentHash = Parameters.SourceData->GetFileContentHash();
 
 			UInterchangeAssetImportData* AssetImportData = nullptr;
 				
-			if (*(Parameters.AssetImportData) != nullptr)
+			if (Parameters.AssetImportData)
 			{
-				UAssetImportData* ParametersAssetImportData = *(Parameters.AssetImportData);
-				AssetImportData = Cast<UInterchangeAssetImportData>(ParametersAssetImportData);
+				AssetImportData = Cast<UInterchangeAssetImportData>(Parameters.AssetImportData);
 				if (!AssetImportData)
 				{
 					AssetImportData = NewObject<UInterchangeAssetImportData>(Parameters.AssetImportDataOuter, NAME_None);
 					TArray<FString> Filenames;
-					ParametersAssetImportData->ExtractFilenames(Filenames);
+					Parameters.AssetImportData->ExtractFilenames(Filenames);
 					TArray<FString> Labels;
-					ParametersAssetImportData->ExtractDisplayLabels(Labels);
+					Parameters.AssetImportData->ExtractDisplayLabels(Labels);
 					if (Filenames.Num() > 1)
 					{
 						for (int32 FileIndex = 0; FileIndex < Filenames.Num(); ++FileIndex)
@@ -94,9 +84,11 @@ namespace UE
 			AssetImportData->NodeUniqueID = Parameters.NodeUniqueID;
 			FObjectDuplicationParameters DupParam(Parameters.NodeContainer, AssetImportData);
 			AssetImportData->NodeContainer = CastChecked<UInterchangeBaseNodeContainer>(StaticDuplicateObjectEx(DupParam));
-			//Assign the asset import data to the imported asset.
-			*(Parameters.AssetImportData) = AssetImportData;
+
+			// Return the asset import data so it can be set on the imported asset.
+			return AssetImportData;
 #endif //#if WITH_EDITORONLY_DATA
+			return nullptr;
 		}
 
 		void FFactoryCommon::ApplyReimportStrategyToAsset(const EReimportStrategyFlags ReimportStrategyFlags
diff --git a/Engine/Plugins/Runtime/Interchange/Source/Import/Private/InterchangeImportCommon.h b/Engine/Plugins/Runtime/Interchange/Source/Import/Private/InterchangeImportCommon.h
--- a/Engine/Plugins/Runtime/Interchange/Source/Import/Private/InterchangeImportCommon.h
+++ b/Engine/Plugins/Runtime/Interchange/Source/Import/Private/InterchangeImportCommon.h
@@ -25,19 +25,13 @@ namespace UE
 			struct FUpdateImportAssetDataParameters
 			{
 				UObject* AssetImportDataOuter = nullptr;
-				UAssetImportData** AssetImportData = nullptr;
+				UAssetImportData* AssetImportData = nullptr;
 				const UInterchangeSourceData* SourceData = nullptr;
 				FString NodeUniqueID;
 				UInterchangeBaseNodeContainer* NodeContainer;
 
 				FUpdateImportAssetDataParameters(UObject* InAssetImportDataOuter
-																	, UAssetImportData** InAssetImportData
-																	, const UInterchangeSourceData* InSourceData
-																	, FString InNodeUniqueID
-																	, UInterchangeBaseNodeContainer* InNodeContainer);
-
-				FUpdateImportAssetDataParameters(UObject* InAssetImportDataOuter
-																	, TObjectPtr<UAssetImportData>* InAssetImportData
+																	, UAssetImportData* InAssetImportData
 																	, const UInterchangeSourceData* InSourceData
 																	, FString InNodeUniqueID
 																	, UInterchangeBaseNodeContainer* InNodeContainer);
@@ -46,8 +40,9 @@ namespace UE
 			/**
 			 * Update the AssetImportData source file of the specified asset in the parameters. Also update the node container and the node unique id.
 			 * If the AssetImportData is null it will create one. If the AssetImportData is not an UInterchangeAssetImportData it will create a new one.
+			 * @return The source data that should be stored on the asset or nullptr if a parameter is invalid
 			 */
-			static void UpdateImportAssetData(FUpdateImportAssetDataParameters& Parameters);
+			static UAssetImportData* UpdateImportAssetData(FUpdateImportAssetDataParameters& Parameters);
 
 			/**
 			 * Apply the current strategy to the PipelineAssetNode
diff --git a/Engine/Plugins/Runtime/Interchange/Source/Import/Private/Material/InterchangeMaterialFactory.cpp b/Engine/Plugins/Runtime/Interchange/Source/Import/Private/Material/InterchangeMaterialFactory.cpp
--- a/Engine/Plugins/Runtime/Interchange/Source/Import/Private/Material/InterchangeMaterialFactory.cpp
+++ b/Engine/Plugins/Runtime/Interchange/Source/Import/Private/Material/InterchangeMaterialFactory.cpp
@@ -295,11 +295,12 @@ void UInterchangeMaterialFactory::PostImportGameThreadCallback(const FPostImport
 		UMaterialInterface* ImportedMaterial = CastChecked<UMaterialInterface>(Arguments.ImportedObject);
 
 		UE::Interchange::FFactoryCommon::FUpdateImportAssetDataParameters UpdateImportAssetDataParameters(ImportedMaterial
-																										  , &ImportedMaterial->AssetImportData
+																										  , ImportedMaterial->AssetImportData
 																										  , Arguments.SourceData
 																										  , Arguments.NodeUniqueID
 																										  , Arguments.NodeContainer);
-		UE::Interchange::FFactoryCommon::UpdateImportAssetData(UpdateImportAssetDataParameters);
+
+		ImportedMaterial->AssetImportData = UE::Interchange::FFactoryCommon::UpdateImportAssetData(UpdateImportAssetDataParameters);
 	}
 #endif
 }
\ No newline at end of file
diff --git a/Engine/Plugins/Runtime/Interchange/Source/Import/Private/Mesh/InterchangeSkeletalMeshFactory.cpp b/Engine/Plugins/Runtime/Interchange/Source/Import/Private/Mesh/InterchangeSkeletalMeshFactory.cpp
--- a/Engine/Plugins/Runtime/Interchange/Source/Import/Private/Mesh/InterchangeSkeletalMeshFactory.cpp
+++ b/Engine/Plugins/Runtime/Interchange/Source/Import/Private/Mesh/InterchangeSkeletalMeshFactory.cpp
@@ -1295,11 +1295,12 @@ void UInterchangeSkeletalMeshFactory::PostImportGameThreadCallback(const FPostIm
 
 		UAssetImportData* ImportDataPtr = SkeletalMesh->GetAssetImportData();
 		UE::Interchange::FFactoryCommon::FUpdateImportAssetDataParameters UpdateImportAssetDataParameters(SkeletalMesh
-																										  , &ImportDataPtr
+																										  , ImportDataPtr
 																										  , Arguments.SourceData
 																										  , Arguments.NodeUniqueID
 																										  , Arguments.NodeContainer);
-		UE::Interchange::FFactoryCommon::UpdateImportAssetData(UpdateImportAssetDataParameters);
+
+		ImportDataPtr = UE::Interchange::FFactoryCommon::UpdateImportAssetData(UpdateImportAssetDataParameters);
 		SkeletalMesh->SetAssetImportData(ImportDataPtr);
 	}
 #endif
diff --git a/Engine/Plugins/Runtime/Interchange/Source/Import/Private/Texture/InterchangeTextureFactory.cpp b/Engine/Plugins/Runtime/Interchange/Source/Import/Private/Texture/InterchangeTextureFactory.cpp
--- a/Engine/Plugins/Runtime/Interchange/Source/Import/Private/Texture/InterchangeTextureFactory.cpp
+++ b/Engine/Plugins/Runtime/Interchange/Source/Import/Private/Texture/InterchangeTextureFactory.cpp
@@ -272,11 +272,11 @@ void UInterchangeTextureFactory::PostImportGameThreadCallback(const FPostImportG
 		UTexture* ImportedTexture = CastChecked<UTexture>(Arguments.ImportedObject);
 		
 		UE::Interchange::FFactoryCommon::FUpdateImportAssetDataParameters UpdateImportAssetDataParameters(ImportedTexture
-																						 , &ImportedTexture->AssetImportData
+																						 , ImportedTexture->AssetImportData
 																						 , Arguments.SourceData
 																						 , Arguments.NodeUniqueID
 																						 , Arguments.NodeContainer);
-		UE::Interchange::FFactoryCommon::UpdateImportAssetData(UpdateImportAssetDataParameters);
+		ImportedTexture->AssetImportData = UE::Interchange::FFactoryCommon::UpdateImportAssetData(UpdateImportAssetDataParameters);
 	}
 #endif
 }
\ No newline at end of file
