Removed unessary ensure in FUpdateImportAssetDataParameters constructor.

diff --git a/Engine/Plugins/Runtime/Interchange/Source/Import/Private/InterchangeImportCommon.cpp b/Engine/Plugins/Runtime/Interchange/Source/Import/Private/InterchangeImportCommon.cpp
--- a/Engine/Plugins/Runtime/Interchange/Source/Import/Private/InterchangeImportCommon.cpp
+++ b/Engine/Plugins/Runtime/Interchange/Source/Import/Private/InterchangeImportCommon.cpp
@@ -24,7 +24,6 @@ namespace UE
 			, NodeContainer(InNodeContainer)
 		{
 			ensure(AssetImportDataOuter);
-			ensure(AssetImportData);
 			ensure(SourceData);
 			ensure(!NodeUniqueID.IsEmpty());
 			ensure(NodeContainer);
