Add an optionnal specialized type array to the interchange scene node, so we can tell a scene node act like a special type.
This is useful for type that do not have any payload and justa couple properties like LOD group, joint, the scene node will be able to return his specialized type.

diff --git a/Engine/Plugins/Editor/InterchangeEditor/Source/Pipelines/InterchangeEditorPipelines.Build.cs b/Engine/Plugins/Editor/InterchangeEditor/Source/Pipelines/InterchangeEditorPipelines.Build.cs
--- a/Engine/Plugins/Editor/InterchangeEditor/Source/Pipelines/InterchangeEditorPipelines.Build.cs
+++ b/Engine/Plugins/Editor/InterchangeEditor/Source/Pipelines/InterchangeEditorPipelines.Build.cs
@@ -24,6 +24,7 @@ namespace UnrealBuildTool.Rules
 					"EditorStyle",
 					"InputCore",
 					"MainFrame",
+					"Projects",
 					"PropertyEditor",
 					"Slate",
 					"SlateCore",
diff --git a/Engine/Plugins/Editor/InterchangeEditor/Source/Pipelines/Private/InterchangeEditorPipelineStyle.h b/Engine/Plugins/Editor/InterchangeEditor/Source/Pipelines/Private/InterchangeEditorPipelineStyle.h
new file mode 100644
--- /dev/null
+++ b/Engine/Plugins/Editor/InterchangeEditor/Source/Pipelines/Private/InterchangeEditorPipelineStyle.h
@@ -0,0 +1,39 @@
+// Copyright Epic Games, Inc. All Rights Reserved.
+
+#pragma once
+
+#include "Interfaces/IPluginManager.h"
+#include "Brushes/SlateImageBrush.h"
+#include "Styling/SlateStyle.h"
+#include "Styling/SlateStyleRegistry.h"
+
+/**
+ * Implements the visual style of the text asset editor UI.
+ */
+class FInterchangeEditorPipelineStyle
+	: public FSlateStyleSet
+{
+public:
+
+	/** Default constructor. */
+	FInterchangeEditorPipelineStyle()
+		: FSlateStyleSet("InterchangeEditorPipelineStyle")
+	{
+		const FVector2D Icon16x16(16.0f, 16.0f);
+
+		const FString BaseDir = IPluginManager::Get().FindPlugin("InterchangeEditor")->GetBaseDir();
+		SetContentRoot(BaseDir / TEXT("Content"));
+
+		FSlateImageBrush* LodBrush16 = new FSlateImageBrush(RootToContentDir(TEXT("Resources/Interchange_Lod_Icon_16"), TEXT(".png")), Icon16x16);
+		Set("SceneGraphIcon.LodGroup", LodBrush16);
+		FSlateImageBrush* JointBrush16 = new FSlateImageBrush(RootToContentDir(TEXT("Resources/Interchange_Joint_Icon_16"), TEXT(".png")), Icon16x16);
+		Set("SceneGraphIcon.Joint", JointBrush16);
+		FSlateStyleRegistry::RegisterSlateStyle(*this);
+	}
+
+	/** Destructor. */
+	~FInterchangeEditorPipelineStyle()
+	{
+		FSlateStyleRegistry::UnRegisterSlateStyle(*this);
+	}
+};
diff --git a/Engine/Plugins/Editor/InterchangeEditor/Source/Pipelines/Private/InterchangeEditorPipelinesModule.cpp b/Engine/Plugins/Editor/InterchangeEditor/Source/Pipelines/Private/InterchangeEditorPipelinesModule.cpp
--- a/Engine/Plugins/Editor/InterchangeEditor/Source/Pipelines/Private/InterchangeEditorPipelinesModule.cpp
+++ b/Engine/Plugins/Editor/InterchangeEditor/Source/Pipelines/Private/InterchangeEditorPipelinesModule.cpp
@@ -5,6 +5,7 @@
 #include "CoreMinimal.h"
 #include "Engine/Engine.h"
 #include "InterchangeEditorPipelineDetails.h"
+#include "InterchangeEditorPipelineStyle.h"
 #include "InterchangeGraphInspectorPipeline.h"
 #include "InterchangeManager.h"
 #include "Misc/CoreDelegates.h"
@@ -18,6 +19,10 @@ class FInterchangeEditorPipelinesModule : public IInterchangeEditorPipelinesModu
 	/** IModuleInterface implementation */
 	virtual void StartupModule() override;
 	virtual void ShutdownModule() override;
+
+private:
+	/** Pointer to the style set to use for the UI. */
+	TSharedPtr<ISlateStyle> InterchangeEditorPipelineStyle = nullptr;
 };
 
 IMPLEMENT_MODULE(FInterchangeEditorPipelinesModule, InterchangeEditorPipelines)
@@ -44,6 +49,11 @@ void FInterchangeEditorPipelinesModule::StartupModule()
 
 	ClassesToUnregisterOnShutdown.Add(UInterchangeBaseNode::StaticClass()->GetFName());
 	PropertyEditorModule.RegisterCustomClassLayout(ClassesToUnregisterOnShutdown.Last(), FOnGetDetailCustomizationInstance::CreateStatic(&FInterchangeBaseNodeDetailsCustomization::MakeInstance));
+	
+	if (!InterchangeEditorPipelineStyle.IsValid())
+	{
+		InterchangeEditorPipelineStyle = MakeShared<FInterchangeEditorPipelineStyle>();
+	}
 }
 
 
@@ -57,6 +67,8 @@ void FInterchangeEditorPipelinesModule::ShutdownModule()
 			PropertyEditorModule->UnregisterCustomClassLayout(ClassName);
 		}
 	}
+
+	InterchangeEditorPipelineStyle = nullptr;
 }
 
 
diff --git a/Engine/Plugins/Editor/InterchangeEditor/Source/Pipelines/Private/SInterchangeGraphInspectorWindow.cpp b/Engine/Plugins/Editor/InterchangeEditor/Source/Pipelines/Private/SInterchangeGraphInspectorWindow.cpp
--- a/Engine/Plugins/Editor/InterchangeEditor/Source/Pipelines/Private/SInterchangeGraphInspectorWindow.cpp
+++ b/Engine/Plugins/Editor/InterchangeEditor/Source/Pipelines/Private/SInterchangeGraphInspectorWindow.cpp
@@ -3,9 +3,9 @@
 
 #include "EditorStyleSet.h"
 #include "Framework/MultiBox/MultiBoxBuilder.h"
+#include "GameFramework/Actor.h"
 #include "IDetailsView.h"
 #include "IDocumentation.h"
-#include "GameFramework/Actor.h"
 #include "Modules/ModuleManager.h"
 #include "Nodes/InterchangeBaseNode.h"
 #include "Nodes/InterchangeBaseNodeContainer.h"
@@ -82,17 +82,27 @@ public:
 	{
 		InterchangeNode = InArgs._InterchangeNode;
 		InterchangeBaseNodeContainer = InArgs._InterchangeBaseNodeContainer;
-
 		//This is suppose to always be valid
 		check(InterchangeNode);
 		check(InterchangeBaseNodeContainer);
 
-		UClass* IconClass = AActor::StaticClass();
-		if (InterchangeNode->GetAssetClass())
+		UClass* IconClass = InterchangeNode->GetAssetClass();
+		const FSlateBrush* TypeIcon = nullptr;
+		FName IconName = InterchangeNode->GetIconName();
+		if (IconClass)
+		{
+			TypeIcon = FSlateIconFinder::FindIconBrushForClass(IconClass);
+		}
+		else if (IconName != NAME_None)
+		{
+			const FSlateIcon SlateIcon = FSlateIconFinder::FindIcon(IconName);
+			TypeIcon = SlateIcon.GetOptionalIcon();
+		}
+		
+		if (!TypeIcon)
 		{
-			IconClass = InterchangeNode->GetAssetClass();
+			TypeIcon = FSlateIconFinder::FindIconBrushForClass(AActor::StaticClass());
 		}
-		const FSlateBrush* ClassIcon = FSlateIconFinder::FindIconBrushForClass(IconClass);
 
 		//Prepare the tooltip
 		FString Tooltip = InterchangeNode->GetDisplayLabel();
@@ -122,8 +132,8 @@ public:
 			.Padding(0.0f, 2.0f, 6.0f, 2.0f)
 			[
 				SNew(SImage)
-				.Image(ClassIcon)
-				.Visibility(ClassIcon != FEditorStyle::GetDefaultBrush() ? EVisibility::Visible : EVisibility::Collapsed)
+				.Image(TypeIcon)
+				.Visibility(TypeIcon != FEditorStyle::GetDefaultBrush() ? EVisibility::Visible : EVisibility::Collapsed)
 			]
 			+ SHorizontalBox::Slot()
 			.FillWidth(1.0f)
diff --git a/Engine/Plugins/Editor/InterchangeEditor/Source/Pipelines/Public/InterchangeEditorPipelinesModule.h b/Engine/Plugins/Editor/InterchangeEditor/Source/Pipelines/Public/InterchangeEditorPipelinesModule.h
--- a/Engine/Plugins/Editor/InterchangeEditor/Source/Pipelines/Public/InterchangeEditorPipelinesModule.h
+++ b/Engine/Plugins/Editor/InterchangeEditor/Source/Pipelines/Public/InterchangeEditorPipelinesModule.h
@@ -5,6 +5,7 @@
 #include "CoreMinimal.h"
 #include "Modules/ModuleInterface.h"
 #include "Modules/ModuleManager.h"
+#include "Styling/ISlateStyle.h"
 
 
 /**
diff --git a/Engine/Plugins/Runtime/Interchange/Source/Nodes/Public/InterchangeSceneNode.h b/Engine/Plugins/Runtime/Interchange/Source/Nodes/Public/InterchangeSceneNode.h
--- a/Engine/Plugins/Runtime/Interchange/Source/Nodes/Public/InterchangeSceneNode.h
+++ b/Engine/Plugins/Runtime/Interchange/Source/Nodes/Public/InterchangeSceneNode.h
@@ -7,6 +7,24 @@
 
 #include "InterchangeSceneNode.generated.h"
 
+//Interchange namespace
+//Interchange namespace
+namespace UE
+{
+	namespace Interchange
+	{
+
+		struct FSceneNodeStaticData : public FBaseNodeStaticData
+		{
+			static const FString& GetNodeSpecializeTypeBaseKey()
+			{
+				static FString SceneNodeSpecializeType_BaseKey = TEXT("SceneNodeSpecializeType");
+				return SceneNodeSpecializeType_BaseKey;
+			}
+		};
+
+	}//ns Interchange
+}//ns UE
 
 UCLASS(BlueprintType)
 class INTERCHANGENODES_API UInterchangeSceneNode : public UInterchangeBaseNode
@@ -17,6 +35,7 @@ public:
 	UInterchangeSceneNode()
 	:UInterchangeBaseNode()
 	{
+		NodeSpecializeTypes.Initialize(Attributes, UE::Interchange::FSceneNodeStaticData::GetNodeSpecializeTypeBaseKey());
 	}
 
 	/**
@@ -28,11 +47,79 @@ public:
 		return TypeName;
 	}
 
+	virtual FString GetKeyDisplayName(const UE::Interchange::FAttributeKey& NodeAttributeKey) const override
+	{
+		FString KeyDisplayName = NodeAttributeKey.Key;
+		if (NodeAttributeKey.Key.Equals(UE::Interchange::FSceneNodeStaticData::GetNodeSpecializeTypeBaseKey()))
+		{
+			KeyDisplayName = TEXT("Specialized Type Count");
+			return KeyDisplayName;
+		}
+		else if (NodeAttributeKey.Key.StartsWith(UE::Interchange::FSceneNodeStaticData::GetNodeSpecializeTypeBaseKey()))
+		{
+			KeyDisplayName = TEXT("Specialized type index ");
+			const FString IndexKey = UE::Interchange::FNameAttributeArrayHelper::IndexKey();
+			int32 IndexPosition = NodeAttributeKey.Key.Find(IndexKey) + IndexKey.Len();
+			if (IndexPosition < NodeAttributeKey.Key.Len())
+			{
+				KeyDisplayName += NodeAttributeKey.Key.RightChop(IndexPosition);
+			}
+			return KeyDisplayName;
+		}
+		return Super::GetKeyDisplayName(NodeAttributeKey);
+	}
+
 	virtual FGuid GetHash() const override
 	{
 		return Attributes->GetStorageHash();
 	}
 
+	/**
+	 * Icon name are simply create by adding "InterchangeIcon_" in front of the specialized type. If there is no special type the function will return NAME_None which will use the default icon.
+	 */
+	virtual FName GetIconName() const override
+	{
+		FString SpecializedType;
+		GetSpecializedType(0, SpecializedType);
+		if (SpecializedType.IsEmpty())
+		{
+			return NAME_None;
+		}
+		SpecializedType = TEXT("SceneGraphIcon.") + SpecializedType;
+		return FName(*SpecializedType);
+	}
+
+	/** Specialized type are scene node special type like (Joint or LODGroup).*/
+	UFUNCTION(BlueprintCallable, Category = "Interchange | Node | Scene")
+	int32 GetSpecializedTypeCount() const
+	{
+		return NodeSpecializeTypes.GetCount();
+	}
+
+	UFUNCTION(BlueprintCallable, Category = "Interchange | Node | Scene")
+	void GetSpecializedType(const int32 Index, FString& OutSpecializedType) const
+	{
+		NodeSpecializeTypes.GetName(Index, OutSpecializedType);
+	}
+
+	UFUNCTION(BlueprintCallable, Category = "Interchange | Node | Scene")
+	void GetSpecializedTypes(TArray<FString>& OutSpecializedTypes) const
+	{
+		NodeSpecializeTypes.GetNames(OutSpecializedTypes);
+	}
+
+	UFUNCTION(BlueprintCallable, Category = "Interchange | Node | Scene")
+	bool AddSpecializedType(const FString& SpecializedType)
+	{
+		return NodeSpecializeTypes.AddName(SpecializedType);
+	}
+
+	UFUNCTION(BlueprintCallable, Category = "Interchange | Node | Scene")
+	bool RemoveSpecializedType(const FString& SpecializedType)
+	{
+		return NodeSpecializeTypes.RemoveName(SpecializedType);
+	}
+
 	/** Return false if the Attribute was not set previously.*/
 	UFUNCTION(BlueprintCallable, Category = "Interchange | Node | Scene")
 	bool GetCustomLocalTransform(FTransform& AttributeValue) const
@@ -63,4 +150,6 @@ private:
 	//Scene Attribute Keys
 	const UE::Interchange::FAttributeKey Macro_CustomLocalTransformKey = UE::Interchange::FAttributeKey(TEXT("LocalTransform"));
 	const UE::Interchange::FAttributeKey Macro_CustomGlobalTransformKey = UE::Interchange::FAttributeKey(TEXT("GlobalTransform"));
+
+	UE::Interchange::FNameAttributeArrayHelper NodeSpecializeTypes;
 };
diff --git a/Engine/Source/Runtime/Interchange/Core/Public/Nodes/InterchangeBaseNode.h b/Engine/Source/Runtime/Interchange/Core/Public/Nodes/InterchangeBaseNode.h
--- a/Engine/Source/Runtime/Interchange/Core/Public/Nodes/InterchangeBaseNode.h
+++ b/Engine/Source/Runtime/Interchange/Core/Public/Nodes/InterchangeBaseNode.h
@@ -217,6 +217,40 @@ namespace UE
 				return NameCount;
 			}
 
+			void GetName(const int32 Index, FString& OutName) const
+			{
+				OutName.Empty();
+				//The Class must be initialise properly before we can use it
+				TSharedPtr<FAttributeStorage, ESPMode::ThreadSafe> AttributePtr = Attributes.Pin();
+				if (!ensure(AttributePtr.IsValid()))
+				{
+					return;
+				}
+				int32 NameCount = 0;
+				if (!AttributePtr->ContainAttribute(GetKeyCount()))
+				{
+					return;
+				}
+
+				FAttributeStorage::TAttributeHandle<int32> Handle = AttributePtr->GetAttributeHandle<int32>(GetKeyCount());
+				if (!Handle.IsValid())
+				{
+					return;
+				}
+				Handle.Get(NameCount);
+				if (Index >= NameCount)
+				{
+					return;
+				}
+				FAttributeKey DepIndexKey = GetIndexKey(Index);
+				FAttributeStorage::TAttributeHandle<FString> HandleName = AttributePtr->GetAttributeHandle<FString>(DepIndexKey);
+				if (!HandleName.IsValid())
+				{
+					return;
+				}
+				HandleName.Get(OutName);
+			}
+
 			void GetNames(TArray<FString>& OutNames) const
 			{
 				//The Class must be initialise properly before we can use it
@@ -767,6 +801,14 @@ public:
 	 */
 	virtual FString GetTypeName() const;
 
+	/**
+	 * Icon name are simply create by adding "InterchangeIcon_" in front of the specialized type. If there is no special type the function will return NAME_None which will use the default icon.
+	 */
+	virtual FName GetIconName() const
+	{
+		return NAME_None;
+	}
+
 	virtual FString GetKeyDisplayName(const UE::Interchange::FAttributeKey& NodeAttributeKey) const
 	{
 		FString KeyDisplayName = NodeAttributeKey.ToString();
