Refactor Interchange graph nodes to follow the new design. We now have translated nodes and factory nodes. The translated node are the source description and the factory nodes are the pipeline description telling the factory how to import the translated source nodes.

diff --git a/Engine/Config/BaseEngine.ini b/Engine/Config/BaseEngine.ini
--- a/Engine/Config/BaseEngine.ini
+++ b/Engine/Config/BaseEngine.ini
@@ -2647,7 +2647,7 @@ bDisableActionLimit=false
 ActionLimit=100
 
 [/Script/InterchangeEngine.InterchangeProjectSettings]
-PipelineStack=/Script/InterchangePipelines.InterchangeGenericMeshPipeline
+PipelineStack=/Script/InterchangePipelines.InterchangeGenericAssetsPipeline
 
 [AlternateTextureCompression]
 TextureCompressionFormat="TextureFormatOodle"
diff --git a/Engine/Plugins/Editor/InterchangeEditor/Source/Pipelines/Private/InterchangeEditorPipelineDetails.cpp b/Engine/Plugins/Editor/InterchangeEditor/Source/Pipelines/Private/InterchangeEditorPipelineDetails.cpp
--- a/Engine/Plugins/Editor/InterchangeEditor/Source/Pipelines/Private/InterchangeEditorPipelineDetails.cpp
+++ b/Engine/Plugins/Editor/InterchangeEditor/Source/Pipelines/Private/InterchangeEditorPipelineDetails.cpp
@@ -57,12 +57,31 @@ void FInterchangeBaseNodeDetailsCustomization::CustomizeDetails( IDetailLayoutBu
 		return;
 	}
 
-	TArray<UE::Interchange::FAttributeKey> AttributeKeys;
+	TArray< UE::Interchange::FAttributeKey> AttributeKeys;
 	InterchangeBaseNode->GetAttributeKeys(AttributeKeys);
-	IDetailCategoryBuilder& AttributeCategory = DetailBuilder.EditCategory("Attributes", FText::GetEmpty(), ECategoryPriority::Important);
+
+	TMap<FString, TArray< UE::Interchange::FAttributeKey>> AttributesPerCategory;
 	for (UE::Interchange::FAttributeKey& AttributeKey : AttributeKeys)
 	{
-		AddAttributeRow(AttributeKey, AttributeCategory);
+		if (InterchangeBaseNode->ShouldHideAttribute(AttributeKey))
+		{
+			//Skip attribute we should hide
+			continue;
+		}
+		const FString CategoryName = InterchangeBaseNode->GetAttributeCategory(AttributeKey);
+		TArray< UE::Interchange::FAttributeKey>& CategoryAttributeKeys = AttributesPerCategory.FindOrAdd(CategoryName);
+		CategoryAttributeKeys.Add(AttributeKey);
+	}
+
+	//Add all categories
+	for (TPair<FString, TArray< UE::Interchange::FAttributeKey>>& CategoryAttributesPair : AttributesPerCategory)
+	{
+		FName CategoryName = FName(*CategoryAttributesPair.Key);
+		IDetailCategoryBuilder& AttributeCategoryBuilder = DetailBuilder.EditCategory(CategoryName, FText::GetEmpty());
+		for (UE::Interchange::FAttributeKey& AttributeKey : CategoryAttributesPair.Value)
+		{
+			AddAttributeRow(AttributeKey, AttributeCategoryBuilder);
+		}
 	}
 }
 
@@ -159,6 +178,12 @@ void FInterchangeBaseNodeDetailsCustomization::AddAttributeRow(UE::Interchange::
 		}
 		break;
 
+		case UE::Interchange::EAttributeTypes::Box:
+		{
+			BuildBoxValueContent(AttributeCategory, AttributeKey);
+		}
+		break;
+
 		default:
 		{
 			FText AttributeName = FText::FromString(InterchangeBaseNode->GetKeyDisplayName(AttributeKey));
@@ -464,6 +489,98 @@ void FInterchangeBaseNodeDetailsCustomization::BuildTransformValueContent(IDetai
     ];
 }
 
+void FInterchangeBaseNodeDetailsCustomization::BuildBoxValueContent(IDetailCategoryBuilder& AttributeCategory, UE::Interchange::FAttributeKey& AttributeKey)
+{
+	UE::Interchange::EAttributeTypes AttributeType = InterchangeBaseNode->GetAttributeType(AttributeKey);
+	check(AttributeType == UE::Interchange::EAttributeTypes::Box);
+
+	{
+		const UE::Interchange::FAttributeStorage::TAttributeHandle<FBox> AttributeHandle = InterchangeBaseNode->GetAttributeHandle<FBox>(AttributeKey);
+		if (!AttributeHandle.IsValid())
+		{
+			CreateInvalidHandleRow(AttributeCategory, AttributeKey);
+			return;
+		}
+	}
+
+	const bool bAdvancedProperty = false;
+	const FString GroupName = InterchangeBaseNode->GetKeyDisplayName(AttributeKey);
+	IDetailGroup& Group = AttributeCategory.AddGroup(FName(*GroupName), FText::FromString(GroupName), bAdvancedProperty);
+	FDetailWidgetRow& GroupHeaderRow = Group.HeaderRow();
+	GroupHeaderRow.NameContent().Widget = SNew(SBox)
+	[
+		CreateNameWidget(AttributeKey)
+	];
+
+	auto GetMinimumValue = [](UInterchangeBaseNode* BaseNode, UE::Interchange::FAttributeKey& Key)->FVector
+	{
+		FBox BoxValue;
+		const UE::Interchange::FAttributeStorage::TAttributeHandle<FBox> AttributeHandle = BaseNode->GetAttributeHandle<FBox>(Key);
+		if (AttributeHandle.IsValid())
+		{
+			AttributeHandle.Get(BoxValue);
+		}
+		return BoxValue.Min;
+	};
+
+	auto SetMinimumValue = [](UInterchangeBaseNode* BaseNode, UE::Interchange::FAttributeKey& Key, const FVector& Value)
+	{
+		FBox BoxValue;
+		UE::Interchange::FAttributeStorage::TAttributeHandle<FBox> AttributeHandle = BaseNode->GetAttributeHandle<FBox>(Key);
+		if (AttributeHandle.IsValid())
+		{
+			AttributeHandle.Get(BoxValue);
+			BoxValue.Min = Value;
+			AttributeHandle.Set(BoxValue);
+		}
+	};
+
+	auto GetMaximumValue = [](UInterchangeBaseNode* BaseNode, UE::Interchange::FAttributeKey& Key)->FVector
+	{
+		FBox BoxValue;
+		const UE::Interchange::FAttributeStorage::TAttributeHandle<FBox> AttributeHandle = BaseNode->GetAttributeHandle<FBox>(Key);
+		if (AttributeHandle.IsValid())
+		{
+			AttributeHandle.Get(BoxValue);
+		}
+		return BoxValue.Max;
+	};
+
+	auto SetMaximumValue = [](UInterchangeBaseNode* BaseNode, UE::Interchange::FAttributeKey& Key, const FVector& Value)
+	{
+		FBox BoxValue;
+		UE::Interchange::FAttributeStorage::TAttributeHandle<FBox> AttributeHandle = BaseNode->GetAttributeHandle<FBox>(Key);
+		if (AttributeHandle.IsValid())
+		{
+			AttributeHandle.Get(BoxValue);
+			BoxValue.Max = Value;
+			AttributeHandle.Set(BoxValue);
+		}
+	};
+
+	const FString MinimumVectorName = TEXT("Minimum");
+	Group.AddWidgetRow()
+	.NameContent()
+	[
+		CreateSimpleNameWidget(MinimumVectorName)
+	]
+	.ValueContent()
+	[
+		CreateVectorWidget(GetMinimumValue, SetMinimumValue, AttributeKey)
+	];
+
+	const FString MaximumVectorName = TEXT("Maximum");
+	Group.AddWidgetRow()
+	.NameContent()
+	[
+		CreateSimpleNameWidget(MaximumVectorName)
+	]
+	.ValueContent()
+	[
+		CreateVectorWidget(GetMaximumValue, SetMaximumValue, AttributeKey)
+	];
+}
+
 void FInterchangeBaseNodeDetailsCustomization::CreateInvalidHandleRow(IDetailCategoryBuilder& AttributeCategory, UE::Interchange::FAttributeKey& AttributeKey) const
 {
 	const FString InvalidAttributeHandle = TEXT("Invalid Attribute Handle!");
diff --git a/Engine/Plugins/Editor/InterchangeEditor/Source/Pipelines/Private/InterchangeEditorPipelineDetails.h b/Engine/Plugins/Editor/InterchangeEditor/Source/Pipelines/Private/InterchangeEditorPipelineDetails.h
--- a/Engine/Plugins/Editor/InterchangeEditor/Source/Pipelines/Private/InterchangeEditorPipelineDetails.h
+++ b/Engine/Plugins/Editor/InterchangeEditor/Source/Pipelines/Private/InterchangeEditorPipelineDetails.h
@@ -58,6 +58,10 @@ private:
 	//Any numeric field are read/write
 	void BuildTransformValueContent(IDetailCategoryBuilder& AttributeCategory, UE::Interchange::FAttributeKey& AttributeKey);
 
+	//FBox: Add a group to the category that contain the Minimum corner position(FVector), and the maximum corner position(FVector)
+	//Any numeric field are read/write
+	void BuildBoxValueContent(IDetailCategoryBuilder& AttributeCategory, UE::Interchange::FAttributeKey& AttributeKey);
+
 	/**
 	 * AttributeType Build functions End
 	 ******************************************************************/
diff --git a/Engine/Plugins/Editor/InterchangeEditor/Source/Pipelines/Private/InterchangeEditorPipelineStyle.h b/Engine/Plugins/Editor/InterchangeEditor/Source/Pipelines/Private/InterchangeEditorPipelineStyle.h
--- a/Engine/Plugins/Editor/InterchangeEditor/Source/Pipelines/Private/InterchangeEditorPipelineStyle.h
+++ b/Engine/Plugins/Editor/InterchangeEditor/Source/Pipelines/Private/InterchangeEditorPipelineStyle.h
@@ -28,6 +28,10 @@ public:
 		Set("SceneGraphIcon.LodGroup", LodBrush16);
 		FSlateImageBrush* JointBrush16 = new FSlateImageBrush(RootToContentDir(TEXT("Resources/Interchange_Joint_Icon_16"), TEXT(".png")), Icon16x16);
 		Set("SceneGraphIcon.Joint", JointBrush16);
+		FSlateImageBrush* StaticMeshBrush16 = new FSlateImageBrush(RootToContentDir(TEXT("Resources/Interchange_StaticMesh_Icon_16"), TEXT(".png")), Icon16x16);
+		Set("MeshIcon.Static", StaticMeshBrush16);
+		FSlateImageBrush* SkeletalMeshBrush16 = new FSlateImageBrush(RootToContentDir(TEXT("Resources/Interchange_SkeletalMesh_Icon_16"), TEXT(".png")), Icon16x16);
+		Set("MeshIcon.Skinned", SkeletalMeshBrush16);
 		FSlateStyleRegistry::RegisterSlateStyle(*this);
 	}
 
diff --git a/Engine/Plugins/Editor/InterchangeEditor/Source/Pipelines/Private/InterchangeGraphInspectorPipeline.cpp b/Engine/Plugins/Editor/InterchangeEditor/Source/Pipelines/Private/InterchangeGraphInspectorPipeline.cpp
--- a/Engine/Plugins/Editor/InterchangeEditor/Source/Pipelines/Private/InterchangeGraphInspectorPipeline.cpp
+++ b/Engine/Plugins/Editor/InterchangeEditor/Source/Pipelines/Private/InterchangeGraphInspectorPipeline.cpp
@@ -3,8 +3,8 @@
 
 #include "CoreMinimal.h"
 #include "InterchangeSkeletalMeshLodDataNode.h"
-#include "InterchangeSkeletalMeshNode.h"
-#include "InterchangeSkeletonNode.h"
+#include "InterchangeSkeletalMeshFactoryNode.h"
+#include "InterchangeSkeletonFactoryNode.h"
 #include "InterchangeSourceData.h"
 #include "Interfaces/IMainFrameModule.h"
 #include "Misc/Paths.h"
diff --git a/Engine/Plugins/Editor/InterchangeEditor/Source/Pipelines/Private/SInterchangeGraphInspectorWindow.cpp b/Engine/Plugins/Editor/InterchangeEditor/Source/Pipelines/Private/SInterchangeGraphInspectorWindow.cpp
--- a/Engine/Plugins/Editor/InterchangeEditor/Source/Pipelines/Private/SInterchangeGraphInspectorWindow.cpp
+++ b/Engine/Plugins/Editor/InterchangeEditor/Source/Pipelines/Private/SInterchangeGraphInspectorWindow.cpp
@@ -158,7 +158,7 @@ private:
 
 	void RecursiveSetImport(UInterchangeBaseNode* Node, bool bState)
 	{
-		TArray<FString> Childrens = InterchangeBaseNodeContainer->GetNodeChildrenUIDs(Node->GetUniqueID());
+		TArray<FString> Childrens = InterchangeBaseNodeContainer->GetNodeChildrenUids(Node->GetUniqueID());
 		for (const FString& ChildID : Childrens)
 		{
 			UInterchangeBaseNode* ChildNode = InterchangeBaseNodeContainer->GetNode(ChildID);
@@ -197,7 +197,7 @@ TSharedRef< ITableRow > SInterchangeGraphInspectorTreeView::OnGenerateRowGraphIn
 }
 void SInterchangeGraphInspectorTreeView::OnGetChildrenGraphInspectorTreeView(UInterchangeBaseNode* InParent, TArray< UInterchangeBaseNode* >& OutChildren)
 {
-	TArray<FString> Childrens = InterchangeBaseNodeContainer->GetNodeChildrenUIDs(InParent->GetUniqueID());
+	TArray<FString> Childrens = InterchangeBaseNodeContainer->GetNodeChildrenUids(InParent->GetUniqueID());
 	for (const FString& ChildID : Childrens)
 	{
 		UInterchangeBaseNode* ChildNode = InterchangeBaseNodeContainer->GetNode(ChildID);
@@ -209,7 +209,7 @@ void SInterchangeGraphInspectorTreeView::OnGetChildrenGraphInspectorTreeView(UIn
 
 void SInterchangeGraphInspectorTreeView::RecursiveSetImport(UInterchangeBaseNode* Node, bool bState)
 {
-	TArray<FString> Childrens = InterchangeBaseNodeContainer->GetNodeChildrenUIDs(Node->GetUniqueID());
+	TArray<FString> Childrens = InterchangeBaseNodeContainer->GetNodeChildrenUids(Node->GetUniqueID());
 	for (const FString& ChildID : Childrens)
 	{
 		UInterchangeBaseNode* ChildNode = InterchangeBaseNodeContainer->GetNode(ChildID);
@@ -235,7 +235,7 @@ void SInterchangeGraphInspectorTreeView::OnToggleSelectAll(ECheckBoxState CheckT
 void SInterchangeGraphInspectorTreeView::RecursiveSetExpand(UInterchangeBaseNode* Node, bool ExpandState)
 {
 	SetItemExpansion(Node, ExpandState);
-	TArray<FString> Childrens = InterchangeBaseNodeContainer->GetNodeChildrenUIDs(Node->GetUniqueID());
+	TArray<FString> Childrens = InterchangeBaseNodeContainer->GetNodeChildrenUids(Node->GetUniqueID());
 	for (const FString& ChildID : Childrens)
 	{
 		UInterchangeBaseNode* ChildNode = InterchangeBaseNodeContainer->GetNode(ChildID);
diff --git a/Engine/Plugins/Enterprise/DatasmithImporter/Source/DatasmithImporter/Private/DatasmithTextureImporter.cpp b/Engine/Plugins/Enterprise/DatasmithImporter/Source/DatasmithImporter/Private/DatasmithTextureImporter.cpp
--- a/Engine/Plugins/Enterprise/DatasmithImporter/Source/DatasmithImporter/Private/DatasmithTextureImporter.cpp
+++ b/Engine/Plugins/Enterprise/DatasmithImporter/Source/DatasmithImporter/Private/DatasmithTextureImporter.cpp
@@ -34,6 +34,7 @@
 #include "ObjectTools.h"
 #include "RHI.h"
 #include "InterchangeSourceData.h"
+#include "InterchangeTextureFactoryNode.h"
 #include "InterchangeTextureNode.h"
 
 #define LOCTEXT_NAMESPACE "DatasmithTextureImport"
@@ -416,12 +417,24 @@ bool UDatasmithTexturePipeline::ExecutePreImportPipeline(UInterchangeBaseNodeCon
 		return false;
 	}
 
-	UInterchangeTextureNode* TextureNode = Cast<UInterchangeTextureNode >( BaseNodeContainer->GetNode( Nodes[0] ) );
-	if (!ensure(TextureNode))
+	UInterchangeTextureNode* TextureTranslatedNode = Cast<UInterchangeTextureNode >( BaseNodeContainer->GetNode( Nodes[0] ) );
+	if (!ensure(TextureTranslatedNode))
 	{
 		return false;
 	}
 
+	FString DisplayLabel = TextureTranslatedNode->GetDisplayLabel();
+	FString NodeUID = UInterchangeTextureFactoryNode::GetTextureFactoryNodeUidFromTextureNodeUid(TextureTranslatedNode->GetUniqueID());
+	UInterchangeTextureFactoryNode* TextureFactoryNode = NewObject<UInterchangeTextureFactoryNode>(BaseNodeContainer, NAME_None);
+	if (!ensure(TextureFactoryNode))
+	{
+		return false;
+	}
+	//Creating a UTexture2D
+	TextureFactoryNode->InitializeTextureNode(NodeUID, DisplayLabel, UTexture2D::StaticClass()->GetName(), TextureTranslatedNode->GetDisplayLabel());
+	TextureFactoryNode->SetCustomTranslatedTextureNodeUid(TextureTranslatedNode->GetUniqueID());
+	BaseNodeContainer->AddNode(TextureFactoryNode);
+
 	TOptional< bool > bFlipNormalMapGreenChannel;
 	TOptional< TextureMipGenSettings > MipGenSettings;
 	TOptional< TextureGroup > LODGroup;
@@ -496,42 +509,42 @@ bool UDatasmithTexturePipeline::ExecutePreImportPipeline(UInterchangeBaseNodeCon
 		bSrgb = false;
 	}
 
-	TextureNode->SetCustomAddressX( (TextureAddress)TextureElement->GetTextureAddressX() );
-	TextureNode->SetCustomAddressY( (TextureAddress)TextureElement->GetTextureAddressY() );
+	TextureFactoryNode->SetCustomAddressX( (TextureAddress)TextureElement->GetTextureAddressX() );
+	TextureFactoryNode->SetCustomAddressY( (TextureAddress)TextureElement->GetTextureAddressY() );
 
 	if ( bSrgb.IsSet() )
 	{
-		TextureNode->SetCustomSRGB( bSrgb.GetValue() );
+		TextureFactoryNode->SetCustomSRGB( bSrgb.GetValue() );
 	}
 
 	if ( bFlipNormalMapGreenChannel.IsSet() )
 	{
-		TextureNode->SetCustombFlipGreenChannel( bFlipNormalMapGreenChannel.GetValue() );
+		TextureFactoryNode->SetCustombFlipGreenChannel( bFlipNormalMapGreenChannel.GetValue() );
 	}
 
 	if ( MipGenSettings.IsSet() )
 	{
-		TextureNode->SetCustomMipGenSettings( MipGenSettings.GetValue() );
+		TextureFactoryNode->SetCustomMipGenSettings( MipGenSettings.GetValue() );
 	}
 
 	if ( LODGroup.IsSet() )
 	{
-		TextureNode->SetCustomLODGroup( LODGroup.GetValue() );
+		TextureFactoryNode->SetCustomLODGroup( LODGroup.GetValue() );
 	}
 
 	if ( CompressionSettings.IsSet() )
 	{
-		TextureNode->SetCustomLODGroup( CompressionSettings.GetValue() );
+		TextureFactoryNode->SetCustomLODGroup( CompressionSettings.GetValue() );
 	}
 
 	if ( RGBCurve.IsSet() )
 	{
-		TextureNode->SetCustomAdjustRGBCurve( RGBCurve.GetValue() );
+		TextureFactoryNode->SetCustomAdjustRGBCurve( RGBCurve.GetValue() );
 	}
 
 	if ( TexFilter.IsSet() )
 	{
-		TextureNode->SetCustomFilter( TexFilter.GetValue() );
+		TextureFactoryNode->SetCustomFilter( TexFilter.GetValue() );
 	}
 
 	return true;
diff --git a/Engine/Plugins/Runtime/Interchange/Source/Import/Private/Material/InterchangeMaterialFactory.cpp b/Engine/Plugins/Runtime/Interchange/Source/Import/Private/Material/InterchangeMaterialFactory.cpp
--- a/Engine/Plugins/Runtime/Interchange/Source/Import/Private/Material/InterchangeMaterialFactory.cpp
+++ b/Engine/Plugins/Runtime/Interchange/Source/Import/Private/Material/InterchangeMaterialFactory.cpp
@@ -3,7 +3,9 @@
 
 #include "InterchangeImportCommon.h"
 #include "InterchangeImportLog.h"
+#include "InterchangeMaterialFactoryNode.h"
 #include "InterchangeMaterialNode.h"
+#include "InterchangeTextureFactoryNode.h"
 #include "InterchangeSourceData.h"
 #include "Materials/Material.h"
 #include "Materials/MaterialExpressionScalarParameter.h"
@@ -37,13 +39,20 @@ UObject* UInterchangeMaterialFactory::CreateEmptyAsset(const FCreateAssetParams&
 		return nullptr;
 	}
 
-	const UInterchangeMaterialNode* MaterialNode = Cast<UInterchangeMaterialNode>(Arguments.AssetNode);
-	if (MaterialNode == nullptr)
+	const UInterchangeMaterialFactoryNode* MaterialFactoryNode = Cast<UInterchangeMaterialFactoryNode>(Arguments.AssetNode);
+	if (MaterialFactoryNode == nullptr)
+	{
+		return nullptr;
+	}
+	FString MaterialTranslatedNodeUid;
+	MaterialFactoryNode->GetCustomTranslatedMaterialNodeUid(MaterialTranslatedNodeUid);
+	const UInterchangeMaterialNode* MaterialTranslatedNode = Cast<UInterchangeMaterialNode>(Arguments.NodeContainer->GetNode(MaterialTranslatedNodeUid));
+	if (MaterialTranslatedNode == nullptr)
 	{
 		return nullptr;
 	}
 
-	const UClass* MaterialClass = MaterialNode->GetAssetClass();
+	const UClass* MaterialClass = MaterialFactoryNode->GetAssetClass();
 	if (!ensure(MaterialClass && MaterialClass->IsChildOf(GetFactoryClass())))
 	{
 		return nullptr;
@@ -86,14 +95,21 @@ UObject* UInterchangeMaterialFactory::CreateAsset(const UInterchangeMaterialFact
 	{
 		return nullptr;
 	}
-
-	const UInterchangeMaterialNode* MaterialNode = Cast<UInterchangeMaterialNode>(Arguments.AssetNode);
-	if (MaterialNode == nullptr)
+	
+	const UInterchangeMaterialFactoryNode* MaterialFactoryNode = Cast<UInterchangeMaterialFactoryNode>(Arguments.AssetNode);
+	if (MaterialFactoryNode == nullptr)
+	{
+		return nullptr;
+	}
+	FString MaterialTranslatedNodeUid;
+	MaterialFactoryNode->GetCustomTranslatedMaterialNodeUid(MaterialTranslatedNodeUid);
+	const UInterchangeMaterialNode* MaterialTranslatedNode = Cast<UInterchangeMaterialNode>(Arguments.NodeContainer->GetNode(MaterialTranslatedNodeUid));
+	if (MaterialTranslatedNode == nullptr)
 	{
 		return nullptr;
 	}
 
-	const UClass* MaterialClass = MaterialNode->GetAssetClass();
+	const UClass* MaterialClass = MaterialFactoryNode->GetAssetClass();
 	check(MaterialClass && MaterialClass->IsChildOf(GetFactoryClass()));
 
 	// create an asset if it doesn't exist
@@ -130,27 +146,31 @@ UObject* UInterchangeMaterialFactory::CreateAsset(const UInterchangeMaterialFact
 			
 			if (Material)
 			{
-				auto ApplyInputParameter = [&MaterialNode, &Material, &Arguments](EInterchangeMaterialNodeParameterName ParameterName, FExpressionInput& MaterialInput, const FVector2D& Location)
+				auto ApplyInputParameter = [&MaterialFactoryNode, &MaterialTranslatedNode, &Material, &Arguments](EInterchangeMaterialNodeParameterName ParameterName, FExpressionInput& MaterialInput, const FVector2D& Location)
 				{
-					FString OutTextureUID;
+					FString OutTextureUid;
 					int32 OutUVSetIndex = 0;
 					float ScaleU = 0.0f;
 					float ScaleV = 0.0f;
 					FVector OutVectorParameter = FVector(0.0f);
 					float OutFloatParameter = 0.0f;
-					if (MaterialNode->GetTextureParameterData(ParameterName, OutTextureUID, OutUVSetIndex, ScaleU, ScaleV))
+					if (MaterialTranslatedNode->GetTextureParameterData(ParameterName, OutTextureUid, OutUVSetIndex, ScaleU, ScaleV))
 					{
+						OutTextureUid = UInterchangeTextureFactoryNode::GetTextureFactoryNodeUidFromTextureNodeUid(OutTextureUid);
 						UTexture* TextureReference = nullptr;
 						if (Arguments.NodeContainer)
 						{
-							FString NodeUniqueID = MaterialNode->GetUniqueID();
-							TArray<FString> MaterialDependencies;
-							MaterialNode->GetDependecies(MaterialDependencies);
-							int32 DependenciesCount = MaterialDependencies.Num();
+							TArray<FString> MaterialTextureDependencies;
+							MaterialFactoryNode->GetTextureDependencies(MaterialTextureDependencies);
+							int32 DependenciesCount = MaterialTextureDependencies.Num();
 							for (int32 DependIndex = 0; DependIndex < DependenciesCount; ++DependIndex)
 							{
-								const UInterchangeBaseNode* DepNode = Arguments.NodeContainer->GetNode(MaterialDependencies[DependIndex]);
-								if (DepNode && DepNode->GetUniqueID() == OutTextureUID && DepNode->ReferenceObject.IsAsset())
+								if (OutTextureUid != MaterialTextureDependencies[DependIndex])
+								{
+									continue;
+								}
+								const UInterchangeTextureFactoryNode* DepNode = Cast<UInterchangeTextureFactoryNode>(Arguments.NodeContainer->GetNode(MaterialTextureDependencies[DependIndex]));
+								if (DepNode && DepNode->ReferenceObject.IsAsset())
 								{
 									//Use Resolve object so we just look for in memory UObject
 									UObject* TextureObject = DepNode->ReferenceObject.ResolveObject();
@@ -206,7 +226,7 @@ UObject* UInterchangeMaterialFactory::CreateAsset(const UInterchangeMaterialFact
 							}
 						}
 					}
-					else if (MaterialNode->GetVectorParameterData(ParameterName, OutVectorParameter))
+					else if (MaterialTranslatedNode->GetVectorParameterData(ParameterName, OutVectorParameter))
 					{
 						UMaterialExpressionVectorParameter* MyColorExpression = NewObject<UMaterialExpressionVectorParameter>(Material);
 						Material->Expressions.Add( MyColorExpression );
@@ -219,7 +239,7 @@ UObject* UInterchangeMaterialFactory::CreateAsset(const UInterchangeMaterialFact
 						MyColorExpression->MaterialExpressionEditorX = FMath::TruncToInt(Location.X);
 						MyColorExpression->MaterialExpressionEditorY = FMath::TruncToInt(Location.Y);
 					}
-					else if (MaterialNode->GetScalarParameterData(ParameterName, OutFloatParameter))
+					else if (MaterialTranslatedNode->GetScalarParameterData(ParameterName, OutFloatParameter))
 					{
 						UMaterialExpressionScalarParameter* MyScalarExpression = NewObject<UMaterialExpressionScalarParameter>(Material);
 						Material->Expressions.Add(MyScalarExpression);
@@ -261,7 +281,7 @@ UObject* UInterchangeMaterialFactory::CreateAsset(const UInterchangeMaterialFact
 			}
 
 			/** Apply all MaterialNode custom attributes to the material asset */
-			MaterialNode->ApplyAllCustomAttributeToAsset(Material);
+			MaterialFactoryNode->ApplyAllCustomAttributeToAsset(Material);
 
 			//TODO support material instance
 		}
diff --git a/Engine/Plugins/Runtime/Interchange/Source/Import/Private/Mesh/InterchangeSkeletalMeshFactory.cpp b/Engine/Plugins/Runtime/Interchange/Source/Import/Private/Mesh/InterchangeSkeletalMeshFactory.cpp
--- a/Engine/Plugins/Runtime/Interchange/Source/Import/Private/Mesh/InterchangeSkeletalMeshFactory.cpp
+++ b/Engine/Plugins/Runtime/Interchange/Source/Import/Private/Mesh/InterchangeSkeletalMeshFactory.cpp
@@ -6,11 +6,10 @@
 #include "GPUSkinPublicDefs.h"
 #include "InterchangeImportCommon.h"
 #include "InterchangeImportLog.h"
-#include "InterchangeJointNode.h"
 #include "InterchangeMaterialNode.h"
 #include "InterchangeSkeletalMeshLodDataNode.h"
-#include "InterchangeSkeletalMeshNode.h"
-#include "InterchangeSkeletonNode.h"
+#include "InterchangeSkeletalMeshFactoryNode.h"
+#include "InterchangeSkeletonFactoryNode.h"
 #include "InterchangeSourceData.h"
 #include "InterchangeTranslatorBase.h"
 #include "Mesh/InterchangeSkeletalMeshPayload.h"
@@ -28,7 +27,7 @@
 
 #endif //WITH_EDITORONLY_DATA
 
-#if WITH_EDITOR
+#if 0//WITH_EDITOR
 namespace UE
 {
 	namespace Interchange
@@ -998,7 +997,7 @@ UClass* UInterchangeSkeletalMeshFactory::GetFactoryClass() const
 
 UObject* UInterchangeSkeletalMeshFactory::CreateEmptyAsset(const FCreateAssetParams& Arguments) const
 {
-#if !WITH_EDITOR || !WITH_EDITORONLY_DATA
+#if 1//!WITH_EDITOR || !WITH_EDITORONLY_DATA
 
 	UE_LOG(LogInterchangeImport, Error, TEXT("Cannot import skeletalMesh asset in runtime, this is an editor only feature."));
 	return nullptr;
@@ -1046,7 +1045,7 @@ UObject* UInterchangeSkeletalMeshFactory::CreateEmptyAsset(const FCreateAssetPar
 
 UObject* UInterchangeSkeletalMeshFactory::CreateAsset(const UInterchangeSkeletalMeshFactory::FCreateAssetParams& Arguments) const
 {
-#if !WITH_EDITOR || !WITH_EDITORONLY_DATA
+#if 1//!WITH_EDITOR || !WITH_EDITORONLY_DATA
 
 	UE_LOG(LogInterchangeImport, Error, TEXT("Cannot import skeletalMesh asset in runtime, this is an editor only feature."));
 	return nullptr;
diff --git a/Engine/Plugins/Runtime/Interchange/Source/Import/Private/Mesh/InterchangeSkeletonFactory.cpp b/Engine/Plugins/Runtime/Interchange/Source/Import/Private/Mesh/InterchangeSkeletonFactory.cpp
--- a/Engine/Plugins/Runtime/Interchange/Source/Import/Private/Mesh/InterchangeSkeletonFactory.cpp
+++ b/Engine/Plugins/Runtime/Interchange/Source/Import/Private/Mesh/InterchangeSkeletonFactory.cpp
@@ -4,7 +4,7 @@
 #include "Animation/Skeleton.h"
 #include "InterchangeImportCommon.h"
 #include "InterchangeImportLog.h"
-#include "InterchangeSkeletonNode.h"
+#include "InterchangeSkeletonFactoryNode.h"
 #include "InterchangeSourceData.h"
 #include "Nodes/InterchangeBaseNode.h"
 #include "Nodes/InterchangeBaseNodeContainer.h"
@@ -26,7 +26,7 @@ UObject* UInterchangeSkeletonFactory::CreateEmptyAsset(const FCreateAssetParams&
 		return nullptr;
 	}
 
-	const UInterchangeSkeletonNode* SkeletonNode = Cast<UInterchangeSkeletonNode>(Arguments.AssetNode);
+	const UInterchangeSkeletonFactoryNode* SkeletonNode = Cast<UInterchangeSkeletonFactoryNode>(Arguments.AssetNode);
 	if (SkeletonNode == nullptr)
 	{
 		return nullptr;
@@ -70,7 +70,7 @@ UObject* UInterchangeSkeletonFactory::CreateAsset(const UInterchangeSkeletonFact
 		return nullptr;
 	}
 
-	const UInterchangeSkeletonNode* SkeletonNode = Cast<UInterchangeSkeletonNode>(Arguments.AssetNode);
+	const UInterchangeSkeletonFactoryNode* SkeletonNode = Cast<UInterchangeSkeletonFactoryNode>(Arguments.AssetNode);
 	if (SkeletonNode == nullptr)
 	{
 		return nullptr;
diff --git a/Engine/Plugins/Runtime/Interchange/Source/Import/Private/Texture/InterchangeTextureFactory.cpp b/Engine/Plugins/Runtime/Interchange/Source/Import/Private/Texture/InterchangeTextureFactory.cpp
--- a/Engine/Plugins/Runtime/Interchange/Source/Import/Private/Texture/InterchangeTextureFactory.cpp
+++ b/Engine/Plugins/Runtime/Interchange/Source/Import/Private/Texture/InterchangeTextureFactory.cpp
@@ -6,6 +6,7 @@
 #include "InterchangeAssetImportData.h"
 #include "InterchangeImportCommon.h"
 #include "InterchangeImportLog.h"
+#include "InterchangeTextureFactoryNode.h"
 #include "InterchangeTextureNode.h"
 #include "InterchangeTranslatorBase.h"
 #include "Nodes/InterchangeBaseNode.h"
@@ -30,29 +31,49 @@ UObject* UInterchangeTextureFactory::CreateEmptyAsset(const FCreateAssetParams&
 	UObject* Texture = nullptr;
 
 #if WITH_EDITORONLY_DATA
-	if (!Arguments.AssetNode || !Arguments.AssetNode->GetAssetClass()->IsChildOf(UTexture::StaticClass()))
+	if (!Arguments.AssetNode)
 	{
+		UE_LOG(LogInterchangeImport, Error, TEXT("UInterchangeTextureFactory: Asset node parameter is null."));
 		return nullptr;
 	}
 
-	const UInterchangeTextureNode* TextureNode = Cast<UInterchangeTextureNode>(Arguments.AssetNode);
-	if (TextureNode == nullptr)
+	if (!Arguments.AssetNode->GetAssetClass()->IsChildOf(UTexture::StaticClass()))
 	{
+		UE_LOG(LogInterchangeImport, Error, TEXT("UInterchangeTextureFactory: Asset node parameter class doesnt derive from UTexture."));
 		return nullptr;
 	}
 
-	//
-	// Generic 2D Image
-	//
-	const TOptional<FString>& PayLoadKey = TextureNode->GetPayLoadKey();
+	UInterchangeTextureFactoryNode* TextureFactoryNode = Cast<UInterchangeTextureFactoryNode>(Arguments.AssetNode);
+	if (TextureFactoryNode == nullptr)
+	{
+		UE_LOG(LogInterchangeImport, Error, TEXT("UInterchangeTextureFactory: Asset node parameter is not a UInterchangeTextureFactoryNode."));
+		return nullptr;
+	}
+
+	FString TextureNodeUniqueID;
+	const UInterchangeTextureNode* TextureTranslatedNode = nullptr;
+	if (TextureFactoryNode->GetCustomTranslatedTextureNodeUid(TextureNodeUniqueID))
+	{
+		TextureTranslatedNode = Cast<UInterchangeTextureNode>(Arguments.NodeContainer->GetNode(TextureNodeUniqueID));
+	}
+
+	if (!TextureTranslatedNode)
+	{
+		UE_LOG(LogInterchangeImport, Error, TEXT("UInterchangeTextureFactory: Asset factory node (UInterchangeTextureFactoryNode) do not reference a valid UInterchangeTextureNode translated node."));
+		return nullptr;
+	}
+
+	const TOptional<FString>& PayLoadKey = TextureTranslatedNode->GetPayLoadKey();
 	if (!PayLoadKey.IsSet())
 	{
+		UE_LOG(LogInterchangeImport, Error, TEXT("UInterchangeTextureFactory: Texture translated node (UInterchangeTextureNode) doesnt have a payload key."));
 		return nullptr;
 	}
 
-	const UClass* TextureClass = TextureNode->GetAssetClass();
-	if (!ensure(TextureClass && TextureClass->IsChildOf(GetFactoryClass())))
+	const UClass* TextureClass = TextureFactoryNode->GetAssetClass();
+	if (!TextureClass || !TextureClass->IsChildOf(GetFactoryClass()))
 	{
+		UE_LOG(LogInterchangeImport, Error, TEXT("UInterchangeTextureFactory: Texture factory node asset class doesnt match Texture factory supported class."));
 		return nullptr;
 	}
 
@@ -91,20 +112,42 @@ UObject* UInterchangeTextureFactory::CreateAsset(const UInterchangeTextureFactor
 
 #else
 
-	if (!Arguments.AssetNode || !Arguments.AssetNode->GetAssetClass()->IsChildOf(UTexture::StaticClass()))
+	if (!Arguments.AssetNode)
+	{
+		UE_LOG(LogInterchangeImport, Error, TEXT("UInterchangeTextureFactory: Asset node parameter is null."));
+		return nullptr;
+	}
+
+	if (!Arguments.AssetNode->GetAssetClass()->IsChildOf(UTexture::StaticClass()))
 	{
+		UE_LOG(LogInterchangeImport, Error, TEXT("UInterchangeTextureFactory: Asset node parameter class doesnt derive from UTexture."));
 		return nullptr;
 	}
 
-	UInterchangeTextureNode* TextureNode = Cast<UInterchangeTextureNode>(Arguments.AssetNode);
-	if (TextureNode == nullptr)
+	UInterchangeTextureFactoryNode* TextureFactoryNode = Cast<UInterchangeTextureFactoryNode>(Arguments.AssetNode);
+	if (TextureFactoryNode == nullptr)
 	{
+		UE_LOG(LogInterchangeImport, Error, TEXT("UInterchangeTextureFactory: Asset node parameter is not a UInterchangeTextureFactoryNode."));
 		return nullptr;
 	}
 
-	const TOptional<FString>& PayLoadKey = TextureNode->GetPayLoadKey();
+	FString TextureNodeUniqueID;
+	const UInterchangeTextureNode* TextureTranslatedNode = nullptr;
+	if (TextureFactoryNode->GetCustomTranslatedTextureNodeUid(TextureNodeUniqueID))
+	{
+		TextureTranslatedNode = Cast<UInterchangeTextureNode>(Arguments.NodeContainer->GetNode(TextureNodeUniqueID));
+	}
+
+	if (!TextureTranslatedNode)
+	{
+		UE_LOG(LogInterchangeImport, Error, TEXT("UInterchangeTextureFactory: Asset factory node (UInterchangeTextureFactoryNode) do not reference a valid UInterchangeTextureNode translated node."));
+		return nullptr;
+	}
+
+	const TOptional<FString>& PayLoadKey = TextureTranslatedNode->GetPayLoadKey();
 	if (!PayLoadKey.IsSet())
 	{
+		UE_LOG(LogInterchangeImport, Error, TEXT("UInterchangeTextureFactory: Texture translated node (UInterchangeTextureNode) doesnt have a payload key."));
 		return nullptr;
 	}
 
@@ -120,16 +163,21 @@ UObject* UInterchangeTextureFactory::CreateAsset(const UInterchangeTextureFactor
 		 * Possible improvement store when a file was last modified and use that as heuristic to skip the file assuming that the file is same as the one in the previous import
 		 * (how to deal with engine changes or translator changes/hotfix?)
 		 */ 
-		BlockedPayloadData = BlockedTextureTranslator->GetBlockedTexturePayloadData(TextureNode->GetSourceBlocks(), Arguments.SourceData);
+		BlockedPayloadData = BlockedTextureTranslator->GetBlockedTexturePayloadData(TextureTranslatedNode->GetSourceBlocks(), Arguments.SourceData);
 	}
 
 	if(!PayloadData.IsSet() && !BlockedPayloadData.IsSet())
 	{
+		UE_LOG(LogInterchangeImport, Error, TEXT("UInterchangeTextureFactory: Invalid payload."));
 		return nullptr;
 	}
 
-	const UClass* TextureClass = TextureNode->GetAssetClass();
-	check(TextureClass && TextureClass->IsChildOf(GetFactoryClass()));
+	const UClass* TextureClass = TextureFactoryNode->GetAssetClass();
+	if (!TextureClass || !TextureClass->IsChildOf(GetFactoryClass()))
+	{
+		UE_LOG(LogInterchangeImport, Error, TEXT("UInterchangeTextureFactory: Texture factory node asset class doesnt match Texture factory supported class."));
+		return nullptr;
+	}
 
 	// create an asset if it doesn't exist
 	UObject* ExistingAsset = StaticFindObject(nullptr, Arguments.Parent, *Arguments.AssetName);
@@ -151,7 +199,7 @@ UObject* UInterchangeTextureFactory::CreateAsset(const UInterchangeTextureFactor
 
 	if (!Texture)
 	{
-		UE_LOG(LogInterchangeImport, Warning, TEXT("Could not create texture asset %s"), *Arguments.AssetName);
+		UE_LOG(LogInterchangeImport, Error, TEXT("UInterchangeTextureFactory: Could not create texture asset %s"), *Arguments.AssetName);
 		return nullptr;
 	}
 
@@ -227,7 +275,7 @@ UObject* UInterchangeTextureFactory::CreateAsset(const UInterchangeTextureFactor
 		if(!Arguments.ReimportObject)
 		{
 			/** Apply all TextureNode custom attributes to the texture asset */
-			TextureNode->ApplyAllCustomAttributeToAsset(Texture2D);
+			TextureFactoryNode->ApplyAllCustomAttributeToAsset(Texture2D);
 		}
 		else
 		{
@@ -237,11 +285,11 @@ UObject* UInterchangeTextureFactory::CreateAsset(const UInterchangeTextureFactor
 			{
 				PreviousNode = InterchangeAssetImportData->NodeContainer->GetNode(InterchangeAssetImportData->NodeUniqueID);
 			}
-			UInterchangeTextureNode* CurrentNode = NewObject<UInterchangeTextureNode>();
-			UInterchangeBaseNode::CopyStorage(TextureNode, CurrentNode);
+			UInterchangeTextureFactoryNode* CurrentNode = NewObject<UInterchangeTextureFactoryNode>();
+			UInterchangeBaseNode::CopyStorage(TextureFactoryNode, CurrentNode);
 			CurrentNode->FillAllCustomAttributeFromAsset(Texture2D);
 			//Apply reimport strategy
-			UE::Interchange::FFactoryCommon::ApplyReimportStrategyToAsset(Arguments.ReimportStrategyFlags, Texture2D, PreviousNode, CurrentNode, TextureNode);
+			UE::Interchange::FFactoryCommon::ApplyReimportStrategyToAsset(Arguments.ReimportStrategyFlags, Texture2D, PreviousNode, CurrentNode, TextureFactoryNode);
 		}
 		
 		//Getting the file Hash will cache it into the source data
diff --git a/Engine/Plugins/Runtime/Interchange/Source/Import/Private/Texture/InterchangeUDIMTranslator.cpp b/Engine/Plugins/Runtime/Interchange/Source/Import/Private/Texture/InterchangeUDIMTranslator.cpp
--- a/Engine/Plugins/Runtime/Interchange/Source/Import/Private/Texture/InterchangeUDIMTranslator.cpp
+++ b/Engine/Plugins/Runtime/Interchange/Source/Import/Private/Texture/InterchangeUDIMTranslator.cpp
@@ -181,7 +181,7 @@ bool UInterchangeUDIMTranslator::Translate(const UInterchangeSourceData* SourceD
 
 	FString DisplayLabel = PreUDIMName + PostUDIMName;
 
-	TextureNode->InitializeTextureNode(Filename, DisplayLabel, UTexture2D::StaticClass()->GetName());
+	TextureNode->InitializeNode(Filename, DisplayLabel, EInterchangeNodeContainerType::NodeContainerType_TranslatedAsset);
 	TextureNode->SetPayLoadKey(Filename);
 
 	UInterchangeManager& InterchangeManager = UInterchangeManager::GetInterchangeManager();
diff --git a/Engine/Plugins/Runtime/Interchange/Source/Import/Private/Texture/TextureTranslatorUtilities.cpp b/Engine/Plugins/Runtime/Interchange/Source/Import/Private/Texture/TextureTranslatorUtilities.cpp
--- a/Engine/Plugins/Runtime/Interchange/Source/Import/Private/Texture/TextureTranslatorUtilities.cpp
+++ b/Engine/Plugins/Runtime/Interchange/Source/Import/Private/Texture/TextureTranslatorUtilities.cpp
@@ -12,6 +12,7 @@
 bool UE::Interchange::FTextureTranslatorUtilities::Generic2DTextureTranslate(const UInterchangeSourceData* SourceData, UInterchangeBaseNodeContainer& BaseNodeContainer)
 {
 	FString Filename = SourceData->GetFilename();
+	FPaths::NormalizeFilename(Filename);
 	if (!FPaths::FileExists(Filename))
 	{
 		return false;
@@ -25,7 +26,7 @@ bool UE::Interchange::FTextureTranslatorUtilities::Generic2DTextureTranslate(con
 		return false;
 	}
 	//Creating a UTexture2D
-	TextureNode->InitializeTextureNode(NodeUID, DisplayLabel, UTexture2D::StaticClass()->GetName());
+	TextureNode->InitializeNode(NodeUID, DisplayLabel, EInterchangeNodeContainerType::NodeContainerType_TranslatedAsset);
 	TextureNode->SetPayLoadKey(Filename);
 
 	BaseNodeContainer.AddNode(TextureNode);
diff --git a/Engine/Plugins/Runtime/Interchange/Source/Nodes/Public/InterchangeJointNode.h b/Engine/Plugins/Runtime/Interchange/Source/Nodes/Public/InterchangeJointNode.h
deleted file mode 100644
diff --git a/Engine/Plugins/Runtime/Interchange/Source/Nodes/Public/InterchangeMaterialFactoryNode.h b/Engine/Plugins/Runtime/Interchange/Source/Nodes/Public/InterchangeMaterialFactoryNode.h
new file mode 100644
--- /dev/null
+++ b/Engine/Plugins/Runtime/Interchange/Source/Nodes/Public/InterchangeMaterialFactoryNode.h
@@ -0,0 +1,328 @@
+// Copyright Epic Games, Inc. All Rights Reserved.
+
+#pragma once
+
+#include "CoreMinimal.h"
+#include "Nodes/InterchangeBaseNode.h"
+
+#if WITH_ENGINE
+
+#include "Materials/Material.h"
+#include "Materials/MaterialInstance.h"
+#include "Materials/MaterialInterface.h"
+
+#endif
+
+#include "InterchangeMaterialFactoryNode.generated.h"
+
+UCLASS(BlueprintType)
+class INTERCHANGENODES_API UInterchangeMaterialFactoryNode : public UInterchangeBaseNode
+{
+	GENERATED_BODY()
+
+public:
+	UInterchangeMaterialFactoryNode()
+	:UInterchangeBaseNode()
+	{
+#if WITH_ENGINE
+		AssetClass = nullptr;
+#endif
+		TextureDependencies.Initialize(Attributes, TextureDependenciesKey.Key);
+	}
+
+	/**
+	 * Initialize node data
+	 * @param: UniqueID - The uniqueId for this node
+	 * @param DisplayLabel - The name of the node
+	 * @param InAssetClass - The class the material factory will create for this node.
+	 */
+	UFUNCTION(BlueprintCallable, Category = "Interchange | Node | MaterialFactory")
+	void InitializeMaterialNode(const FString& UniqueID, const FString& DisplayLabel, const FString& InAssetClass)
+	{
+		bIsMaterialNodeClassInitialized = false;
+		InitializeNode(UniqueID, DisplayLabel, EInterchangeNodeContainerType::NodeContainerType_FactoryData);
+		FString OperationName = GetTypeName() + TEXT(".SetAssetClassName");
+		InterchangePrivateNodeBase::SetCustomAttribute<FString>(*Attributes, ClassNameAttributeKey, OperationName, InAssetClass);
+		FillAssetClassFromAttribute();
+	}
+
+	/**
+	 * Return the node type name of the class, we use this when reporting error
+	 */
+	virtual FString GetTypeName() const override
+	{
+		const FString TypeName = TEXT("MaterialFactoryNode");
+		return TypeName;
+	}
+
+	/** Get the class this node want to create */
+	UFUNCTION(BlueprintCallable, Category = "Interchange | Node | MaterialFactory")
+	virtual class UClass* GetAssetClass() const override
+	{
+		ensure(bIsMaterialNodeClassInitialized);
+#if WITH_ENGINE
+		return AssetClass.Get() != nullptr ? AssetClass.Get() : UMaterial::StaticClass();
+#else
+		return nullptr;
+#endif
+	}
+
+	virtual FString GetKeyDisplayName(const UE::Interchange::FAttributeKey& NodeAttributeKey) const override
+	{
+		FString KeyDisplayName = NodeAttributeKey.ToString();
+		if (NodeAttributeKey == TextureDependenciesKey)
+		{
+			return KeyDisplayName = TEXT("Texture Dependencies count");
+		}
+		else if (NodeAttributeKey.Key.StartsWith(TextureDependenciesKey.Key))
+		{
+			KeyDisplayName = TEXT("Texture Dependencies Index ");
+			const FString IndexKey = UE::Interchange::FNameAttributeArrayHelper::IndexKey();
+			int32 IndexPosition = NodeAttributeKey.Key.Find(IndexKey) + IndexKey.Len();
+			if (IndexPosition < NodeAttributeKey.Key.Len())
+			{
+				KeyDisplayName += NodeAttributeKey.Key.RightChop(IndexPosition);
+			}
+			return KeyDisplayName;
+		}
+		return Super::GetKeyDisplayName(NodeAttributeKey);
+	}
+
+	virtual FString GetAttributeCategory(const UE::Interchange::FAttributeKey& NodeAttributeKey) const override
+	{
+		if (NodeAttributeKey == Macro_CustomTranslatedMaterialNodeUidKey
+			|| NodeAttributeKey == Macro_CustomMaterialUsageKey
+			|| NodeAttributeKey == Macro_CustomBlendModeKey)
+		{
+			return FString(TEXT("MaterialFactory"));
+		}
+		else if (NodeAttributeKey.Key.StartsWith(TextureDependenciesKey.Key))
+		{
+			return FString(TEXT("TextureDependencies"));
+		}
+		return Super::GetAttributeCategory(NodeAttributeKey);
+	}
+
+	virtual FGuid GetHash() const override
+	{
+		return Attributes->GetStorageHash();
+	}
+
+	static FString GetMaterialFactoryNodeUidFromMaterialNodeUid(const FString& TranslatedNodeUid)
+	{
+		FString NewUid = TEXT("Factory_") + TranslatedNodeUid;
+		return NewUid;
+	}
+
+	/**
+	 * This function allow to retrieve the number of Texture dependencies for this object.
+	 *
+	 */
+	UFUNCTION(BlueprintCallable, Category = "Interchange | Node | MaterialFactory")
+	int32 GetTextureDependeciesCount() const
+	{
+		return TextureDependencies.GetCount();
+	}
+
+	/**
+	 * This function allow to retrieve the Texture dependency for this object.
+	 *
+	 */
+	UFUNCTION(BlueprintCallable, Category = "Interchange | Node | MaterialFactory")
+	void GetTextureDependencies(TArray<FString>& OutDependencies) const
+	{
+		TextureDependencies.GetNames(OutDependencies);
+	}
+
+	/**
+	 * This function allow to retrieve one Texture dependency for this object.
+	 *
+	 */
+	UFUNCTION(BlueprintCallable, Category = "Interchange | Node | MaterialFactory")
+	void GetTextureDependency(const int32 Index, FString& OutDependency) const
+	{
+		TextureDependencies.GetName(Index, OutDependency);
+	}
+
+	/**
+	 * Add one Texture dependency to this object.
+	 */
+	UFUNCTION(BlueprintCallable, Category = "Interchange | Node | MaterialFactory")
+	bool SetTextureDependencyUid(const FString& DependencyUid)
+	{
+		return TextureDependencies.AddName(DependencyUid);
+	}
+
+	/**
+	 * Remove one Texture dependency from this object.
+	 */
+	UFUNCTION(BlueprintCallable, Category = "Interchange | Node | MaterialFactory")
+	bool RemoveTextureDependencyUid(const FString& DependencyUid)
+	{
+		return TextureDependencies.RemoveName(DependencyUid);
+	}
+
+	/**
+	 * Get the translated texture node unique ID.
+	 */
+	UFUNCTION(BlueprintCallable, Category = "Interchange | Node | MaterialFactory")
+	bool GetCustomTranslatedMaterialNodeUid(FString& AttributeValue) const
+	{
+		IMPLEMENT_NODE_ATTRIBUTE_GETTER(TranslatedMaterialNodeUid, FString);
+	}
+
+	/**
+	 * Set the translated texture node unique ID. This is the reference to the node that was create by the translator and this node is needed to get the texture payload.
+	 */
+	UFUNCTION(BlueprintCallable, Category = "Interchange | Node | MaterialFactory")
+	bool SetCustomTranslatedMaterialNodeUid(const FString& AttributeValue)
+	{
+		IMPLEMENT_NODE_ATTRIBUTE_SETTER_NODELEGATE(TranslatedMaterialNodeUid, FString);
+	}
+
+	/** Return false if the Attribute was not set previously.*/
+	UFUNCTION(BlueprintCallable, Category = "Interchange | Node | MaterialFactory")
+	bool GetCustomMaterialUsage(uint8& AttributeValue) const
+	{
+		IMPLEMENT_NODE_ATTRIBUTE_GETTER(MaterialUsage, uint8);
+	}
+
+	UFUNCTION(BlueprintCallable, Category = "Interchange | Node | MaterialFactory")
+	bool SetCustomMaterialUsage(const uint8& AttributeValue, bool bAddApplyDelegate = true)
+	{
+		IMPLEMENT_NODE_ATTRIBUTE_SETTER(UInterchangeMaterialFactoryNode, MaterialUsage, uint8, UMaterialInterface);
+	}
+
+	/** Return false if the Attribute was not set previously.*/
+	UFUNCTION(BlueprintCallable, Category = "Interchange | Node | MaterialFactory")
+	bool GetCustomBlendMode(uint8& AttributeValue) const
+	{
+		IMPLEMENT_NODE_ATTRIBUTE_GETTER(BlendMode, uint8);
+	}
+
+	UFUNCTION(BlueprintCallable, Category = "Interchange | Node | MaterialFactory")
+	bool SetCustomBlendMode(const uint8& AttributeValue, bool bAddApplyDelegate = true)
+	{
+		IMPLEMENT_NODE_ATTRIBUTE_SETTER(UInterchangeMaterialFactoryNode, BlendMode, uint8, UMaterialInterface);
+	}
+
+	virtual void Serialize(FArchive& Ar) override
+	{
+		Super::Serialize(Ar);
+#if WITH_ENGINE
+		if (Ar.IsLoading())
+		{
+			//Make sure the class is properly set when we compile with engine, this will set the
+			//bIsMaterialNodeClassInitialized to true.
+			SetMaterialNodeClassFromClassAttribute();
+		}
+#endif //#if WITH_ENGINE
+	}
+
+private:
+
+	void FillAssetClassFromAttribute()
+	{
+#if WITH_ENGINE
+		FString OperationName = GetTypeName() + TEXT(".GetAssetClassName");
+		FString ClassName;
+		InterchangePrivateNodeBase::GetCustomAttribute<FString>(*Attributes, ClassNameAttributeKey, OperationName, ClassName);
+		if (ClassName.Equals(UMaterial::StaticClass()->GetName()))
+		{
+			AssetClass = UMaterial::StaticClass();
+			bIsMaterialNodeClassInitialized = true;
+		}
+		else if (ClassName.Equals(UMaterialInstance::StaticClass()->GetName()))
+		{
+			AssetClass = UMaterialInstance::StaticClass();
+			bIsMaterialNodeClassInitialized = true;
+		}
+#endif
+	}
+
+	bool SetMaterialNodeClassFromClassAttribute()
+	{
+		if (!bIsMaterialNodeClassInitialized)
+		{
+			FillAssetClassFromAttribute();
+		}
+		return bIsMaterialNodeClassInitialized;
+	}
+
+	const UE::Interchange::FAttributeKey ClassNameAttributeKey = UE::Interchange::FBaseNodeStaticData::ClassTypeAttributeKey();
+	
+	//Translated texture node unique id
+	const UE::Interchange::FAttributeKey Macro_CustomTranslatedMaterialNodeUidKey = UE::Interchange::FAttributeKey(TEXT("TranslatedMaterialNodeUid"));
+
+	//Material Attribute
+	const UE::Interchange::FAttributeKey Macro_CustomMaterialUsageKey = UE::Interchange::FAttributeKey(TEXT("MaterialUsage"));
+	const UE::Interchange::FAttributeKey Macro_CustomBlendModeKey = UE::Interchange::FAttributeKey(TEXT("BlendMode"));
+
+	const UE::Interchange::FAttributeKey TextureDependenciesKey = UE::Interchange::FAttributeKey(TEXT("__TextureDependenciesKey__"));
+
+
+#if WITH_ENGINE
+
+	//Material Ussage is not handle properly, you can have several usage check
+	//Currently I can set only one and retrieve the first, TODO make this work (maybe one key per usage...)
+	bool ApplyCustomMaterialUsageToAsset(UObject * Asset) const
+	{
+		if (!Asset)
+		{
+			return false;
+		}
+		UMaterialInterface* TypedObject = Cast<UMaterialInterface>(Asset);
+		if (!TypedObject)
+		{
+			return false;
+		}
+		uint8 ValueData;
+		if (GetCustomMaterialUsage(ValueData))
+		{
+			TypedObject->CheckMaterialUsage_Concurrent(EMaterialUsage(ValueData));
+			return true;
+		}
+		return false;
+	}
+
+	bool FillCustomMaterialUsageFromAsset(UObject* Asset)
+	{
+		if (!Asset)
+		{
+			return false;
+		}
+		UMaterialInterface* TypedObject = Cast<UMaterialInterface>(Asset);
+		if (!TypedObject)
+		{
+			return false;
+		}
+		const UMaterial* Material = TypedObject->GetMaterial_Concurrent();
+		if (!Material)
+		{
+			return false;
+		}
+		for (int32 Usage = 0; Usage < MATUSAGE_MAX; Usage++)
+		{
+			const EMaterialUsage UsageEnum = (EMaterialUsage)Usage;
+			if (Material->GetUsageByFlag(UsageEnum))
+			{
+				if (SetCustomMaterialUsage((int8)UsageEnum, false))
+				{
+					return true;
+				}
+			}
+		}
+		return false;
+	}
+#endif
+
+	IMPLEMENT_NODE_ATTRIBUTE_APPLY_UOBJECT(BlendMode, uint8, UMaterial, TEnumAsByte<enum EBlendMode>);
+
+protected:
+#if WITH_ENGINE
+	TSubclassOf<UTexture> AssetClass = nullptr;
+#endif
+	bool bIsMaterialNodeClassInitialized = false;
+
+	UE::Interchange::FNameAttributeArrayHelper TextureDependencies;
+};
diff --git a/Engine/Plugins/Runtime/Interchange/Source/Nodes/Public/InterchangeMaterialNode.h b/Engine/Plugins/Runtime/Interchange/Source/Nodes/Public/InterchangeMaterialNode.h
--- a/Engine/Plugins/Runtime/Interchange/Source/Nodes/Public/InterchangeMaterialNode.h
+++ b/Engine/Plugins/Runtime/Interchange/Source/Nodes/Public/InterchangeMaterialNode.h
@@ -20,28 +20,6 @@ namespace UE
 {
 	namespace Interchange
 	{
-		namespace Material
-		{
-			/**
-			 * Remove _SkinXX from a string
-			 */
-			static FString RemoveSkinFromName(const FString& Name)
-			{
-				int32 SkinOffset = Name.Find(TEXT("_skin"), ESearchCase::IgnoreCase, ESearchDir::FromEnd);
-				if (SkinOffset != INDEX_NONE)
-				{
-					FString MaterialNameNoSkin = Name;
-					FString SkinXXNumber = Name.Right(Name.Len() - (SkinOffset + 1)).RightChop(4);
-					if (SkinXXNumber.IsNumeric())
-					{
-						MaterialNameNoSkin = Name.LeftChop(Name.Len() - SkinOffset);
-					}
-					return MaterialNameNoSkin;
-				}
-				return Name;
-			}
-		} //ns Material
-
 		struct FMaterialNodeStaticData : public FBaseNodeStaticData
 		{
 			static const FAttributeKey& PayloadSourceFileKey()
@@ -78,25 +56,7 @@ public:
 	UInterchangeMaterialNode()
 	:UInterchangeBaseNode()
 	{
-#if WITH_ENGINE
-		AssetClass = nullptr;
-#endif
-	}
-
-	/**
-	 * Initialize node data
-	 * @param: UniqueID - The uniqueId for this node
-	 * @param DisplayLabel - The name of the node
-	 * @param InAssetClass - The class the material factory will create for this node.
-	 */
-	UFUNCTION(BlueprintCallable, Category = "Interchange | Node | Material")
-	void InitializeMaterialNode(const FString& UniqueID, const FString& DisplayLabel, const FString& InAssetClass)
-	{
-		bIsMaterialNodeClassInitialized = false;
-		InitializeNode(UniqueID, DisplayLabel);
-		FString OperationName = GetTypeName() + TEXT(".SetAssetClassName");
-		InterchangePrivateNodeBase::SetCustomAttribute<FString>(*Attributes, ClassNameAttributeKey, OperationName, InAssetClass);
-		FillAssetClassFromAttribute();
+		TextureDependencies.Initialize(Attributes, TextureDependenciesKey.Key);
 	}
 
 	/**
@@ -115,19 +75,31 @@ public:
 		{
 			return KeyDisplayName = TEXT("Payload Source Key");
 		}
+		else if (NodeAttributeKey == TextureDependenciesKey)
+		{
+			return KeyDisplayName = TEXT("Texture Dependencies count");
+		}
+		else if (NodeAttributeKey.Key.StartsWith(TextureDependenciesKey.Key))
+		{
+			KeyDisplayName = TEXT("Texture Dependencies Index ");
+			const FString IndexKey = UE::Interchange::FNameAttributeArrayHelper::IndexKey();
+			int32 IndexPosition = NodeAttributeKey.Key.Find(IndexKey) + IndexKey.Len();
+			if (IndexPosition < NodeAttributeKey.Key.Len())
+			{
+				KeyDisplayName += NodeAttributeKey.Key.RightChop(IndexPosition);
+			}
+			return KeyDisplayName;
+		}
 		return Super::GetKeyDisplayName(NodeAttributeKey);
 	}
 
-	/** Get the class this node want to create */
-	UFUNCTION(BlueprintCallable, Category = "Interchange | Node | Material")
-	virtual class UClass* GetAssetClass() const override
+	virtual FString GetAttributeCategory(const UE::Interchange::FAttributeKey& NodeAttributeKey) const override
 	{
-		ensure(bIsMaterialNodeClassInitialized);
-#if WITH_ENGINE
-		return AssetClass.Get() != nullptr ? AssetClass.Get() : UMaterial::StaticClass();
-#else
-		return nullptr;
-#endif
+		if (NodeAttributeKey.Key.StartsWith(TextureDependenciesKey.Key))
+		{
+			return FString(TEXT("TextureDependencies"));
+		}
+		return Super::GetAttributeCategory(NodeAttributeKey);
 	}
 
 	virtual FGuid GetHash() const override
@@ -135,6 +107,54 @@ public:
 		return Attributes->GetStorageHash();
 	}
 
+	/**
+	 * This function allow to retrieve the number of Texture dependencies for this object.
+	 *
+	 */
+	UFUNCTION(BlueprintCallable, Category = "Interchange | Node | Material")
+	int32 GetTextureDependeciesCount() const
+	{
+		return TextureDependencies.GetCount();
+	}
+
+	/**
+	 * This function allow to retrieve the Texture dependency for this object.
+	 *
+	 */
+	UFUNCTION(BlueprintCallable, Category = "Interchange | Node | Material")
+	void GetTextureDependencies(TArray<FString>& OutDependencies) const
+	{
+		TextureDependencies.GetNames(OutDependencies);
+	}
+
+	/**
+	 * This function allow to retrieve one Texture dependency for this object.
+	 *
+	 */
+	UFUNCTION(BlueprintCallable, Category = "Interchange | Node | Material")
+	void GetTextureDependency(const int32 Index, FString& OutDependency) const
+	{
+		TextureDependencies.GetName(Index, OutDependency);
+	}
+
+	/**
+	 * Add one Texture dependency to this object.
+	 */
+	UFUNCTION(BlueprintCallable, Category = "Interchange | Node | Material")
+	bool SetTextureDependencyUid(const FString& DependencyUid)
+	{
+		return TextureDependencies.AddName(DependencyUid);
+	}
+
+	/**
+	 * Remove one Texture dependency from this object.
+	 */
+	UFUNCTION(BlueprintCallable, Category = "Interchange | Node | Material")
+	bool RemoveTextureDependencyUid(const FString& DependencyUid)
+	{
+		return TextureDependencies.RemoveName(DependencyUid);
+	}
+
 	virtual const TOptional<FString> GetPayLoadKey() const
 	{
 		if (!Attributes->ContainAttribute(UE::Interchange::FMaterialNodeStaticData::PayloadSourceFileKey()))
@@ -166,32 +186,6 @@ public:
 		}
 	}
 
-	/** Return false if the Attribute was not set previously.*/
-	UFUNCTION(BlueprintCallable, Category = "Interchange | Node | Material")
-	bool GetCustomMaterialUsage(uint8& AttributeValue) const
-	{
-		IMPLEMENT_NODE_ATTRIBUTE_GETTER(MaterialUsage, uint8);
-	}
-
-	UFUNCTION(BlueprintCallable, Category = "Interchange | Node | Material")
-	bool SetCustomMaterialUsage(const uint8& AttributeValue, bool bAddApplyDelegate = true)
-	{
-		IMPLEMENT_NODE_ATTRIBUTE_SETTER(UInterchangeMaterialNode, MaterialUsage, uint8, UMaterialInterface);
-	}
-
-	/** Return false if the Attribute was not set previously.*/
-	UFUNCTION(BlueprintCallable, Category = "Interchange | Node | Material")
-	bool GetCustomBlendMode(uint8& AttributeValue) const
-	{
-		IMPLEMENT_NODE_ATTRIBUTE_GETTER(BlendMode, uint8);
-	}
-
-	UFUNCTION(BlueprintCallable, Category = "Interchange | Node | Material")
-	bool SetCustomBlendMode(const uint8& AttributeValue, bool bAddApplyDelegate = true)
-	{
-		IMPLEMENT_NODE_ATTRIBUTE_SETTER(UInterchangeMaterialNode, BlendMode, uint8, UMaterialInterface);
-	}
-
 	//////////////////////////////////////////////////////////////////////////
 	/**
 	 * Parameter interface for c++ python and blueprint
@@ -203,26 +197,26 @@ public:
 	/**
 	 * Add a texture parameter for the specified ParameterName.
 	 * @param ParameterName - The parameter we want to set the texture for.
-	 * @param TextureUID - The texture node uniqueID that has the texture we want to set to the specified parameter
+	 * @param TextureUid - The texture node uniqueID that has the texture we want to set to the specified parameter
 	 * @param UVSetNameUseByTextureInput - [Optional] The fbx name of the UV set we want the material input expression use to sample the specified texture.
 	 *                                     It can be empty, the material factory should in that case map it on the UV channel 0
 	 * @note - A parameter name can have only one of the 3 type of input set, the last input type set is the one that will be created.
 	 */
 	UFUNCTION(BlueprintCallable, Category = "Interchange | Node | Material")
-	void AddTextureParameterData(const EInterchangeMaterialNodeParameterName ParameterName, const FString& TextureUID, const int32 UVSetIndex, const float ScaleU, const float ScaleV)
+	void AddTextureParameterData(const EInterchangeMaterialNodeParameterName ParameterName, const FString& TextureUid, const int32 UVSetIndex, const float ScaleU, const float ScaleV)
 	{
 		FParameterData& ParameterData = ParameterDatas.FindOrAdd(ParameterName);
 		ParameterData.bIsTextureParameter = true;
 		ParameterData.bIsVectorParameter = false;
 		ParameterData.bIsScalarParameter = false;
-		ParameterData.TextureUID = TextureUID;
+		ParameterData.TextureUid = TextureUid;
 		ParameterData.UVSetIndex = UVSetIndex;
 		ParameterData.ScaleU = ScaleU;
 		ParameterData.ScaleV = ScaleV;
 	}
 
 	UFUNCTION(BlueprintCallable, Category = "Interchange | Node | Material")
-	bool GetTextureParameterData(const EInterchangeMaterialNodeParameterName ParameterName, FString& OutTextureUID, int32& OutUVSetIndex, float& OutScaleU, float& OutScaleV) const
+	bool GetTextureParameterData(const EInterchangeMaterialNodeParameterName ParameterName, FString& OutTextureUid, int32& OutUVSetIndex, float& OutScaleU, float& OutScaleV) const
 	{
 		if (const FParameterData* ParameterData = ParameterDatas.Find(ParameterName))
 		{
@@ -230,7 +224,7 @@ public:
 			{
 				return false;
 			}
-			OutTextureUID = ParameterData->TextureUID;
+			OutTextureUid = ParameterData->TextureUid;
 			OutUVSetIndex = ParameterData->UVSetIndex;
 			OutScaleU = ParameterData->ScaleU;
 			OutScaleV = ParameterData->ScaleV;
@@ -319,51 +313,15 @@ public:
 				Ar << ParamValue;
 			}
 		}
-
-#if WITH_ENGINE
-		if (Ar.IsLoading())
-		{
-			//Make sure the class is properly set when we compile with engine, this will set the
-			//bIsMaterialNodeClassInitialized to true.
-			SetMaterialNodeClassFromClassAttribute();
-		}
-#endif //#if WITH_ENGINE
 	}
 
 private:
-
-	void FillAssetClassFromAttribute()
-	{
-#if WITH_ENGINE
-		FString OperationName = GetTypeName() + TEXT(".GetAssetClassName");
-		FString ClassName;
-		InterchangePrivateNodeBase::GetCustomAttribute<FString>(*Attributes, ClassNameAttributeKey, OperationName, ClassName);
-		if (ClassName.Equals(UMaterial::StaticClass()->GetName()))
-		{
-			AssetClass = UMaterial::StaticClass();
-			bIsMaterialNodeClassInitialized = true;
-		}
-		else if (ClassName.Equals(UMaterialInstance::StaticClass()->GetName()))
-		{
-			AssetClass = UMaterialInstance::StaticClass();
-			bIsMaterialNodeClassInitialized = true;
-		}
-#endif
-	}
-
-	bool SetMaterialNodeClassFromClassAttribute()
-	{
-		if (!bIsMaterialNodeClassInitialized)
-		{
-			FillAssetClassFromAttribute();
-		}
-		return bIsMaterialNodeClassInitialized;
-	}
+	const UE::Interchange::FAttributeKey TextureDependenciesKey = UE::Interchange::FAttributeKey(TEXT("__TextureDependenciesKey__"));
 
 	struct FParameterData
 	{
 		bool bIsTextureParameter = false;
-		FString TextureUID;
+		FString TextureUid;
 		int32 UVSetIndex = 0; //The "UV set" index we use to set the material sampler input default is channel 0
 		float ScaleU = 1.0f;
 		float ScaleV = 1.0f;
@@ -377,7 +335,7 @@ private:
 		friend FArchive& operator<<(FArchive& Ar, FParameterData& ParameterData)
 		{
 			Ar << ParameterData.bIsTextureParameter;
-			Ar << ParameterData.TextureUID;
+			Ar << ParameterData.TextureUid;
 			Ar << ParameterData.UVSetIndex;
 			Ar << ParameterData.ScaleU;
 			Ar << ParameterData.ScaleV;
@@ -394,73 +352,5 @@ private:
 	//This member is serialize manually in the serialize override
 	TMap<EInterchangeMaterialNodeParameterName, FParameterData> ParameterDatas;
 
-	const UE::Interchange::FAttributeKey ClassNameAttributeKey = UE::Interchange::FBaseNodeStaticData::ClassTypeAttributeKey();
-
-	//Material Attribute
-	const UE::Interchange::FAttributeKey Macro_CustomMaterialUsageKey = UE::Interchange::FAttributeKey(TEXT("MaterialUsage"));
-	const UE::Interchange::FAttributeKey Macro_CustomBlendModeKey = UE::Interchange::FAttributeKey(TEXT("BlendMode"));
-
-#if WITH_ENGINE
-
-	//Material Ussage is not handle properly, you can have several usage check
-	//Currently I can set only one and retrieve the first, TODO make this work (maybe one key per usage...)
-	bool ApplyCustomMaterialUsageToAsset(UObject * Asset) const
-	{
-		if (!Asset)
-		{
-			return false;
-		}
-		UMaterialInterface* TypedObject = Cast<UMaterialInterface>(Asset);
-		if (!TypedObject)
-		{
-			return false;
-		}
-		uint8 ValueData;
-		if (GetCustomMaterialUsage(ValueData))
-		{
-			TypedObject->CheckMaterialUsage_Concurrent(EMaterialUsage(ValueData));
-			return true;
-		}
-		return false;
-	}
-
-	bool FillCustomMaterialUsageFromAsset(UObject* Asset)
-	{
-		if (!Asset)
-		{
-			return false;
-		}
-		UMaterialInterface* TypedObject = Cast<UMaterialInterface>(Asset);
-		if (!TypedObject)
-		{
-			return false;
-		}
-		const UMaterial* Material = TypedObject->GetMaterial_Concurrent();
-		if (!Material)
-		{
-			return false;
-		}
-		for (int32 Usage = 0; Usage < MATUSAGE_MAX; Usage++)
-		{
-			const EMaterialUsage UsageEnum = (EMaterialUsage)Usage;
-			if (Material->GetUsageByFlag(UsageEnum))
-			{
-				if (SetCustomMaterialUsage((int8)UsageEnum, false))
-				{
-					return true;
-				}
-			}
-		}
-		return false;
-	}
-#endif
-
-	IMPLEMENT_NODE_ATTRIBUTE_APPLY_UOBJECT(BlendMode, uint8, UMaterial, TEnumAsByte<enum EBlendMode>);
-
-protected:
-#if WITH_ENGINE
-	TSubclassOf<UTexture> AssetClass = nullptr;
-#endif
-	bool bIsMaterialNodeClassInitialized = false;
-
+	UE::Interchange::FNameAttributeArrayHelper TextureDependencies;
 };
diff --git a/Engine/Plugins/Runtime/Interchange/Source/Nodes/Public/InterchangeMeshNode.h b/Engine/Plugins/Runtime/Interchange/Source/Nodes/Public/InterchangeMeshNode.h
new file mode 100644
--- /dev/null
+++ b/Engine/Plugins/Runtime/Interchange/Source/Nodes/Public/InterchangeMeshNode.h
@@ -0,0 +1,618 @@
+// Copyright Epic Games, Inc. All Rights Reserved.
+
+#pragma once
+
+#include "CoreMinimal.h"
+#include "Nodes/InterchangeBaseNode.h"
+
+#include "InterchangeMeshNode.generated.h"
+
+//Interchange namespace
+namespace UE
+{
+	namespace Interchange
+	{
+
+		struct FMeshNodeStaticData : public FBaseNodeStaticData
+		{
+			static const FAttributeKey& PayloadSourceFileKey()
+			{
+				static FAttributeKey AttributeKey(TEXT("__PayloadSourceFile__"));
+				return AttributeKey;
+			}
+
+			static const FAttributeKey& IsSkinnedMeshKey()
+			{
+				static FAttributeKey AttributeKey(TEXT("__IsSkinnedMeshKey__"));
+				return AttributeKey;
+			}
+
+			static const FAttributeKey& IsBlendShapeKey()
+			{
+				static FAttributeKey AttributeKey(TEXT("__IsBlendShapeKey__"));
+				return AttributeKey;
+			}
+
+			static const FAttributeKey& BlendShapeNameKey()
+			{
+				static FAttributeKey AttributeKey(TEXT("__BlendShapeNameKey__"));
+				return AttributeKey;
+			}
+			static const FString& GetSkeletonDependenciesKey()
+			{
+				static FString Dependencies_BaseKey = TEXT("__MeshSkeletonDependencies__");
+				return Dependencies_BaseKey;
+			}
+			static const FString& GetMaterialDependenciesKey()
+			{
+				static FString Dependencies_BaseKey = TEXT("__MeshMaterialDependencies__");
+				return Dependencies_BaseKey;
+			}
+
+			static const FString& GetShapeDependenciesKey()
+			{
+				static FString Dependencies_BaseKey = TEXT("__MeshShapeDependencies__");
+				return Dependencies_BaseKey;
+			}
+		};
+
+	}//ns Interchange
+}//ns UE
+
+UCLASS(BlueprintType)
+class INTERCHANGENODES_API UInterchangeMeshNode : public UInterchangeBaseNode
+{
+	GENERATED_BODY()
+
+public:
+	UInterchangeMeshNode()
+		:UInterchangeBaseNode()
+	{
+		SkeletonDependencies.Initialize(Attributes, UE::Interchange::FMeshNodeStaticData::GetSkeletonDependenciesKey());
+		MaterialDependencies.Initialize(Attributes, UE::Interchange::FMeshNodeStaticData::GetMaterialDependenciesKey());
+		ShapeDependencies.Initialize(Attributes, UE::Interchange::FMeshNodeStaticData::GetShapeDependenciesKey());
+	}
+
+	virtual FString GetKeyDisplayName(const UE::Interchange::FAttributeKey& NodeAttributeKey) const override
+	{
+		FString KeyDisplayName = NodeAttributeKey.ToString();
+		if (NodeAttributeKey == UE::Interchange::FMeshNodeStaticData::PayloadSourceFileKey())
+		{
+			return KeyDisplayName = TEXT("Payload Source Key");
+		}
+		else if (NodeAttributeKey == UE::Interchange::FMeshNodeStaticData::IsSkinnedMeshKey())
+		{
+			return KeyDisplayName = TEXT("Is a Skinned Mesh");
+		}
+		else if (NodeAttributeKey == UE::Interchange::FMeshNodeStaticData::IsBlendShapeKey())
+		{
+			return KeyDisplayName = TEXT("Is a Blend Shape");
+		}
+		else if (NodeAttributeKey == UE::Interchange::FMeshNodeStaticData::BlendShapeNameKey())
+		{
+			return KeyDisplayName = TEXT("Blend Shape Name");
+		}
+		else if (NodeAttributeKey.Key.Equals(UE::Interchange::FMeshNodeStaticData::GetSkeletonDependenciesKey()))
+		{
+			return KeyDisplayName = TEXT("Skeleton Dependencies count");
+		}
+		else if (NodeAttributeKey.Key.StartsWith(UE::Interchange::FMeshNodeStaticData::GetSkeletonDependenciesKey()))
+		{
+			KeyDisplayName = TEXT("Skeleton Dependencies Index ");
+			const FString IndexKey = UE::Interchange::FNameAttributeArrayHelper::IndexKey();
+			int32 IndexPosition = NodeAttributeKey.Key.Find(IndexKey) + IndexKey.Len();
+			if (IndexPosition < NodeAttributeKey.Key.Len())
+			{
+				KeyDisplayName += NodeAttributeKey.Key.RightChop(IndexPosition);
+			}
+			return KeyDisplayName;
+		}
+		else if (NodeAttributeKey.Key.Equals(UE::Interchange::FMeshNodeStaticData::GetMaterialDependenciesKey()))
+		{
+			return KeyDisplayName = TEXT("Material Dependencies count");
+		}
+		else if (NodeAttributeKey.Key.StartsWith(UE::Interchange::FMeshNodeStaticData::GetMaterialDependenciesKey()))
+		{
+			KeyDisplayName = TEXT("Material Dependencies Index ");
+			const FString IndexKey = UE::Interchange::FNameAttributeArrayHelper::IndexKey();
+			int32 IndexPosition = NodeAttributeKey.Key.Find(IndexKey) + IndexKey.Len();
+			if (IndexPosition < NodeAttributeKey.Key.Len())
+			{
+				KeyDisplayName += NodeAttributeKey.Key.RightChop(IndexPosition);
+			}
+			return KeyDisplayName;
+		}
+		else if (NodeAttributeKey.Key.Equals(UE::Interchange::FMeshNodeStaticData::GetShapeDependenciesKey()))
+		{
+			return KeyDisplayName = TEXT("Shape Dependencies count");
+		}
+		else if (NodeAttributeKey.Key.StartsWith(UE::Interchange::FMeshNodeStaticData::GetShapeDependenciesKey()))
+		{
+			KeyDisplayName = TEXT("Shape Dependencies Index ");
+			const FString IndexKey = UE::Interchange::FNameAttributeArrayHelper::IndexKey();
+			int32 IndexPosition = NodeAttributeKey.Key.Find(IndexKey) + IndexKey.Len();
+			if (IndexPosition < NodeAttributeKey.Key.Len())
+			{
+				KeyDisplayName += NodeAttributeKey.Key.RightChop(IndexPosition);
+			}
+			return KeyDisplayName;
+		}
+		return Super::GetKeyDisplayName(NodeAttributeKey);
+	}
+
+	virtual FString GetAttributeCategory(const UE::Interchange::FAttributeKey& NodeAttributeKey) const override
+	{
+		if (NodeAttributeKey.Key.StartsWith(UE::Interchange::FMeshNodeStaticData::GetSkeletonDependenciesKey()))
+		{
+			return FString(TEXT("SkeletonDependencies"));
+		}
+		else if (NodeAttributeKey.Key.StartsWith(UE::Interchange::FMeshNodeStaticData::GetMaterialDependenciesKey()))
+		{
+			return FString(TEXT("MaterialDependencies"));
+		}
+		else if (NodeAttributeKey.Key.StartsWith(UE::Interchange::FMeshNodeStaticData::GetShapeDependenciesKey()))
+		{
+			return FString(TEXT("ShapeDependencies"));
+		}
+		else if (NodeAttributeKey == Macro_CustomVertexCountKey
+				 || NodeAttributeKey == Macro_CustomPolygonCountKey
+				 || NodeAttributeKey == Macro_CustomBoundingBoxKey
+				 || NodeAttributeKey == Macro_CustomHasVertexNormalKey
+				 || NodeAttributeKey == Macro_CustomHasVertexBinormalKey
+				 || NodeAttributeKey == Macro_CustomHasVertexTangentKey
+				 || NodeAttributeKey == Macro_CustomHasSmoothGroupKey
+				 || NodeAttributeKey == Macro_CustomHasVertexColorKey
+				 || NodeAttributeKey == Macro_CustomUVCountKey)
+		{
+			return FString(TEXT("MeshInfo"));
+		}
+		return Super::GetAttributeCategory(NodeAttributeKey);
+	}
+
+	/**
+	 * Return the node type name of the class, we use this when reporting error
+	 */
+	virtual FString GetTypeName() const override
+	{
+		const FString TypeName = TEXT("MeshNode");
+		return TypeName;
+	}
+
+	virtual FGuid GetHash() const override
+	{
+		return Attributes->GetStorageHash();
+	}
+
+	/**
+	 * Icon name are simply create by adding "InterchangeIcon_" in front of the specialized type. If there is no special type the function will return NAME_None which will use the default icon.
+	 */
+	virtual FName GetIconName() const override
+	{
+		FString MeshIconName = TEXT("MeshIcon.");
+		if (IsSkinnedMesh())
+		{
+			MeshIconName += TEXT("Skinned");
+		}
+		else
+		{
+			MeshIconName += TEXT("Static");
+		}
+		
+		return FName(*MeshIconName);
+	}
+
+	/**
+	 * Return true if this node represent a skinned mesh
+	 */
+	UFUNCTION(BlueprintCallable, Category = "Interchange | Node | Mesh")
+	bool IsSkinnedMesh() const
+	{
+		if (!Attributes->ContainAttribute(UE::Interchange::FMeshNodeStaticData::IsSkinnedMeshKey()))
+		{
+			return false;
+		}
+
+		UE::Interchange::FAttributeStorage::TAttributeHandle<bool> Handle = Attributes->GetAttributeHandle<bool>(UE::Interchange::FMeshNodeStaticData::IsSkinnedMeshKey());
+		if (Handle.IsValid())
+		{
+			bool bValue = false;
+			Handle.Get(bValue);
+			return bValue;
+		}
+		return false;
+	}
+
+	/**
+	 * Set the IsSkinnedMesh attribute to determine if this node represent a skinned mesh.
+	 */
+	UFUNCTION(BlueprintCallable, Category = "Interchange | Node | Mesh")
+	bool SetSkinnedMesh(const bool bIsSkinnedMesh)
+	{
+		UE::Interchange::EAttributeStorageResult Result = Attributes->RegisterAttribute(UE::Interchange::FMeshNodeStaticData::IsSkinnedMeshKey(), bIsSkinnedMesh);
+		if (IsAttributeStorageResultSuccess(Result))
+		{
+			UE::Interchange::FAttributeStorage::TAttributeHandle<bool> Handle = Attributes->GetAttributeHandle<bool>(UE::Interchange::FMeshNodeStaticData::IsSkinnedMeshKey());
+			return Handle.IsValid();
+		}
+		return false;
+	}
+
+	/**
+	 * Return true if this node represent a skinned mesh
+	 */
+	UFUNCTION(BlueprintCallable, Category = "Interchange | Node | Mesh")
+	bool IsBlendShape() const
+	{
+		if (!Attributes->ContainAttribute(UE::Interchange::FMeshNodeStaticData::IsBlendShapeKey()))
+		{
+			return false;
+		}
+
+		UE::Interchange::FAttributeStorage::TAttributeHandle<bool> Handle = Attributes->GetAttributeHandle<bool>(UE::Interchange::FMeshNodeStaticData::IsBlendShapeKey());
+		if (Handle.IsValid())
+		{
+			bool bValue = false;
+			Handle.Get(bValue);
+			return bValue;
+		}
+		return false;
+	}
+
+	/**
+	 * Set the IsBlendShape attribute to determine if this node represent a blend shape.
+	 */
+	UFUNCTION(BlueprintCallable, Category = "Interchange | Node | Mesh")
+	bool SetBlendShape(const bool bIsBlendShape)
+	{
+		UE::Interchange::EAttributeStorageResult Result = Attributes->RegisterAttribute(UE::Interchange::FMeshNodeStaticData::IsBlendShapeKey(), bIsBlendShape);
+		if (IsAttributeStorageResultSuccess(Result))
+		{
+			UE::Interchange::FAttributeStorage::TAttributeHandle<bool> Handle = Attributes->GetAttributeHandle<bool>(UE::Interchange::FMeshNodeStaticData::IsBlendShapeKey());
+			return Handle.IsValid();
+		}
+		return false;
+	}
+
+	/**
+	 * Return true if we successfully query the BlendShapeName attribute
+	 */
+	UFUNCTION(BlueprintCallable, Category = "Interchange | Node | Mesh")
+	bool GetBlendShapeName(FString& OutBlendShapeName) const
+	{
+		OutBlendShapeName.Empty();
+		if (!Attributes->ContainAttribute(UE::Interchange::FMeshNodeStaticData::BlendShapeNameKey()))
+		{
+			return false;
+		}
+
+		UE::Interchange::FAttributeStorage::TAttributeHandle<FString> Handle = Attributes->GetAttributeHandle<FString>(UE::Interchange::FMeshNodeStaticData::BlendShapeNameKey());
+		if (Handle.IsValid())
+		{
+			Handle.Get(OutBlendShapeName);
+			return true;
+		}
+		return false;
+	}
+
+	/**
+	 * Set the BlendShapeName attribute to determine if this node represent a blend shape.
+	 */
+	UFUNCTION(BlueprintCallable, Category = "Interchange | Node | Mesh")
+	bool SetBlendShapeName(const FString& BlendShapeName)
+	{
+		UE::Interchange::EAttributeStorageResult Result = Attributes->RegisterAttribute(UE::Interchange::FMeshNodeStaticData::BlendShapeNameKey(), BlendShapeName);
+		if (IsAttributeStorageResultSuccess(Result))
+		{
+			UE::Interchange::FAttributeStorage::TAttributeHandle<FString> Handle = Attributes->GetAttributeHandle<FString>(UE::Interchange::FMeshNodeStaticData::BlendShapeNameKey());
+			return Handle.IsValid();
+		}
+		return false;
+	}
+
+	/** Mesh node Interface Begin */
+	virtual const TOptional<FString> GetPayLoadKey() const
+	{
+		if (!Attributes->ContainAttribute(UE::Interchange::FMeshNodeStaticData::PayloadSourceFileKey()))
+		{
+			return TOptional<FString>();
+		}
+		UE::Interchange::FAttributeStorage::TAttributeHandle<FString> AttributeHandle = Attributes->GetAttributeHandle<FString>(UE::Interchange::FMeshNodeStaticData::PayloadSourceFileKey());
+		if (!AttributeHandle.IsValid())
+		{
+			return TOptional<FString>();
+		}
+		FString PayloadKey;
+		UE::Interchange::EAttributeStorageResult Result = AttributeHandle.Get(PayloadKey);
+		if (!IsAttributeStorageResultSuccess(Result))
+		{
+			LogAttributeStorageErrors(Result, TEXT("UInterchangeMeshNode.GetPayLoadKey"), UE::Interchange::FMeshNodeStaticData::PayloadSourceFileKey());
+			return TOptional<FString>();
+		}
+		return TOptional<FString>(PayloadKey);
+	}
+
+	UFUNCTION(BlueprintCallable, Category = "Interchange | Node | Mesh")
+	virtual void SetPayLoadKey(const FString& PayloadKey)
+	{
+		UE::Interchange::EAttributeStorageResult Result = Attributes->RegisterAttribute(UE::Interchange::FMeshNodeStaticData::PayloadSourceFileKey(), PayloadKey);
+		if (!IsAttributeStorageResultSuccess(Result))
+		{
+			LogAttributeStorageErrors(Result, TEXT("UInterchangeMeshNode.SetPayLoadKey"), UE::Interchange::FMeshNodeStaticData::PayloadSourceFileKey());
+		}
+	}
+	
+	/** Return false if the Attribute was not set previously.*/
+	UFUNCTION(BlueprintCallable, Category = "Interchange | Node | Mesh")
+	bool GetCustomVertexCount(int32& AttributeValue) const
+	{
+		IMPLEMENT_NODE_ATTRIBUTE_GETTER(VertexCount, int32);
+	}
+
+	UFUNCTION(BlueprintCallable, Category = "Interchange | Node | Mesh")
+	bool SetCustomVertexCount(const int32& AttributeValue)
+	{
+		IMPLEMENT_NODE_ATTRIBUTE_SETTER_NODELEGATE(VertexCount, int32);
+	}
+
+	/** Return false if the Attribute was not set previously.*/
+	UFUNCTION(BlueprintCallable, Category = "Interchange | Node | Mesh")
+	bool GetCustomPolygonCount(int32& AttributeValue) const
+	{
+		IMPLEMENT_NODE_ATTRIBUTE_GETTER(PolygonCount, int32);
+	}
+
+	UFUNCTION(BlueprintCallable, Category = "Interchange | Node | Mesh")
+	bool SetCustomPolygonCount(const int32& AttributeValue)
+	{
+		IMPLEMENT_NODE_ATTRIBUTE_SETTER_NODELEGATE(PolygonCount, int32);
+	}
+
+	/** Return false if the Attribute was not set previously.*/
+	UFUNCTION(BlueprintCallable, Category = "Interchange | Node | Mesh")
+	bool GetCustomBoundingBox(FBox& AttributeValue) const
+	{
+		IMPLEMENT_NODE_ATTRIBUTE_GETTER(BoundingBox, FBox);
+	}
+
+	UFUNCTION(BlueprintCallable, Category = "Interchange | Node | Mesh")
+	bool SetCustomBoundingBox(const FBox& AttributeValue)
+	{
+		IMPLEMENT_NODE_ATTRIBUTE_SETTER_NODELEGATE(BoundingBox, FBox);
+	}
+
+	/** Return false if the Attribute was not set previously.*/
+	UFUNCTION(BlueprintCallable, Category = "Interchange | Node | Mesh")
+	bool GetCustomHasVertexNormal(bool& AttributeValue) const
+	{
+		IMPLEMENT_NODE_ATTRIBUTE_GETTER(HasVertexNormal, bool);
+	}
+
+	UFUNCTION(BlueprintCallable, Category = "Interchange | Node | Mesh")
+	bool SetCustomHasVertexNormal(const bool& AttributeValue)
+	{
+		IMPLEMENT_NODE_ATTRIBUTE_SETTER_NODELEGATE(HasVertexNormal, bool);
+	}
+
+	/** Return false if the Attribute was not set previously.*/
+	UFUNCTION(BlueprintCallable, Category = "Interchange | Node | Mesh")
+	bool GetCustomHasVertexBinormal(bool& AttributeValue) const
+	{
+		IMPLEMENT_NODE_ATTRIBUTE_GETTER(HasVertexBinormal, bool);
+	}
+
+	UFUNCTION(BlueprintCallable, Category = "Interchange | Node | Mesh")
+	bool SetCustomHasVertexBinormal(const bool& AttributeValue)
+	{
+		IMPLEMENT_NODE_ATTRIBUTE_SETTER_NODELEGATE(HasVertexBinormal, bool);
+	}
+
+	/** Return false if the Attribute was not set previously.*/
+	UFUNCTION(BlueprintCallable, Category = "Interchange | Node | Mesh")
+	bool GetCustomHasVertexTangent(bool& AttributeValue) const
+	{
+		IMPLEMENT_NODE_ATTRIBUTE_GETTER(HasVertexTangent, bool);
+	}
+
+	UFUNCTION(BlueprintCallable, Category = "Interchange | Node | Mesh")
+	bool SetCustomHasVertexTangent(const bool& AttributeValue)
+	{
+		IMPLEMENT_NODE_ATTRIBUTE_SETTER_NODELEGATE(HasVertexTangent, bool);
+	}
+
+	/** Return false if the Attribute was not set previously.*/
+	UFUNCTION(BlueprintCallable, Category = "Interchange | Node | Mesh")
+	bool GetCustomHasSmoothGroup(bool& AttributeValue) const
+	{
+		IMPLEMENT_NODE_ATTRIBUTE_GETTER(HasSmoothGroup, bool);
+	}
+
+	UFUNCTION(BlueprintCallable, Category = "Interchange | Node | Mesh")
+	bool SetCustomHasSmoothGroup(const bool& AttributeValue)
+	{
+		IMPLEMENT_NODE_ATTRIBUTE_SETTER_NODELEGATE(HasSmoothGroup, bool);
+	}
+
+	/** Return false if the Attribute was not set previously.*/
+	UFUNCTION(BlueprintCallable, Category = "Interchange | Node | Mesh")
+	bool GetCustomHasVertexColor(bool& AttributeValue) const
+	{
+		IMPLEMENT_NODE_ATTRIBUTE_GETTER(HasVertexColor, bool);
+	}
+
+	UFUNCTION(BlueprintCallable, Category = "Interchange | Node | Mesh")
+	bool SetCustomHasVertexColor(const bool& AttributeValue)
+	{
+		IMPLEMENT_NODE_ATTRIBUTE_SETTER_NODELEGATE(HasVertexColor, bool);
+	}
+
+	/** Return false if the Attribute was not set previously.*/
+	UFUNCTION(BlueprintCallable, Category = "Interchange | Node | Mesh")
+	bool GetCustomUVCount(int32& AttributeValue) const
+	{
+		IMPLEMENT_NODE_ATTRIBUTE_GETTER(UVCount, int32);
+	}
+
+	UFUNCTION(BlueprintCallable, Category = "Interchange | Node | Mesh")
+	bool SetCustomUVCount(const int32& AttributeValue)
+	{
+		IMPLEMENT_NODE_ATTRIBUTE_SETTER_NODELEGATE(UVCount, int32);
+	}
+
+	/**
+	 * This function allow to retrieve the number of skeleton dependencies for this object.
+	 *
+	 */
+	UFUNCTION(BlueprintCallable, Category = "Interchange | Node | Mesh")
+	int32 GetSkeletonDependeciesCount() const
+	{
+		return SkeletonDependencies.GetCount();
+	}
+
+	/**
+	 * This function allow to retrieve the skeleton dependency for this object.
+	 *
+	 */
+	UFUNCTION(BlueprintCallable, Category = "Interchange | Node | Mesh")
+	void GetSkeletonDependencies(TArray<FString>& OutDependencies) const
+	{
+		SkeletonDependencies.GetNames(OutDependencies);
+	}
+
+	/**
+	 * This function allow to retrieve one skeleton dependency for this object.
+	 *
+	 */
+	UFUNCTION(BlueprintCallable, Category = "Interchange | Node | Mesh")
+	void GetSkeletonDependency(const int32 Index, FString& OutDependency) const
+	{
+		SkeletonDependencies.GetName(Index, OutDependency);
+	}
+
+	/**
+	 * Add one skeleton dependency to this object.
+	 */
+	UFUNCTION(BlueprintCallable, Category = "Interchange | Node | Mesh")
+	bool SetSkeletonDependencyUid(const FString& DependencyUid)
+	{ 
+		return SkeletonDependencies.AddName(DependencyUid);
+	}
+
+	/**
+	 * Remove one skeleton dependency from this object.
+	 */
+	UFUNCTION(BlueprintCallable, Category = "Interchange | Node | Mesh")
+	bool RemoveSkeletonDependencyUid(const FString& DependencyUid)
+	{
+		return SkeletonDependencies.RemoveName(DependencyUid);
+	}
+
+	/**
+	 * This function allow to retrieve the number of Material dependencies for this object.
+	 *
+	 */
+	UFUNCTION(BlueprintCallable, Category = "Interchange | Node | Mesh")
+	int32 GetMaterialDependeciesCount() const
+	{
+		return MaterialDependencies.GetCount();
+	}
+
+	/**
+	 * This function allow to retrieve the Material dependency for this object.
+	 *
+	 */
+	UFUNCTION(BlueprintCallable, Category = "Interchange | Node | Mesh")
+	void GetMaterialDependencies(TArray<FString>& OutDependencies) const
+	{
+		MaterialDependencies.GetNames(OutDependencies);
+	}
+
+	/**
+	 * This function allow to retrieve one Material dependency for this object.
+	 *
+	 */
+	UFUNCTION(BlueprintCallable, Category = "Interchange | Node | Mesh")
+	void GetMaterialDependency(const int32 Index, FString& OutDependency) const
+	{
+		MaterialDependencies.GetName(Index, OutDependency);
+	}
+
+	/**
+	 * Add one Material dependency to this object.
+	 */
+	UFUNCTION(BlueprintCallable, Category = "Interchange | Node | Mesh")
+	bool SetMaterialDependencyUid(const FString& DependencyUid)
+	{
+		return MaterialDependencies.AddName(DependencyUid);
+	}
+
+	/**
+	 * Remove one Material dependency from this object.
+	 */
+	UFUNCTION(BlueprintCallable, Category = "Interchange | Node | Mesh")
+	bool RemoveMaterialDependencyUid(const FString& DependencyUid)
+	{
+		return MaterialDependencies.RemoveName(DependencyUid);
+	}
+
+	/**
+	 * This function allow to retrieve the number of Shape dependencies for this object.
+	 *
+	 */
+	UFUNCTION(BlueprintCallable, Category = "Interchange | Node | Mesh")
+	int32 GetShapeDependeciesCount() const
+	{
+		return ShapeDependencies.GetCount();
+	}
+
+	/**
+	 * This function allow to retrieve the Shape dependency for this object.
+	 *
+	 */
+	UFUNCTION(BlueprintCallable, Category = "Interchange | Node | Mesh")
+	void GetShapeDependencies(TArray<FString>& OutDependencies) const
+	{
+		ShapeDependencies.GetNames(OutDependencies);
+	}
+
+	/**
+	 * This function allow to retrieve one Shape dependency for this object.
+	 *
+	 */
+	UFUNCTION(BlueprintCallable, Category = "Interchange | Node | Mesh")
+	void GetShapeDependency(const int32 Index, FString& OutDependency) const
+	{
+		ShapeDependencies.GetName(Index, OutDependency);
+	}
+
+	/**
+	 * Add one Shape dependency to this object.
+	 */
+	UFUNCTION(BlueprintCallable, Category = "Interchange | Node | Mesh")
+	bool SetShapeDependencyUid(const FString& DependencyUid)
+	{
+		return ShapeDependencies.AddName(DependencyUid);
+	}
+
+	/**
+	 * Remove one Shape dependency from this object.
+	 */
+	UFUNCTION(BlueprintCallable, Category = "Interchange | Node | Mesh")
+	bool RemoveShapeDependencyUid(const FString& DependencyUid)
+	{
+		return ShapeDependencies.RemoveName(DependencyUid);
+	}
+private:
+	const UE::Interchange::FAttributeKey Macro_CustomVertexCountKey = UE::Interchange::FAttributeKey(TEXT("VertexCount"));
+	const UE::Interchange::FAttributeKey Macro_CustomPolygonCountKey = UE::Interchange::FAttributeKey(TEXT("PolygonCount"));
+	const UE::Interchange::FAttributeKey Macro_CustomBoundingBoxKey = UE::Interchange::FAttributeKey(TEXT("BoundingBox"));
+	const UE::Interchange::FAttributeKey Macro_CustomHasVertexNormalKey = UE::Interchange::FAttributeKey(TEXT("HasVertexNormal"));
+	const UE::Interchange::FAttributeKey Macro_CustomHasVertexBinormalKey = UE::Interchange::FAttributeKey(TEXT("HasVertexBinormal"));
+	const UE::Interchange::FAttributeKey Macro_CustomHasVertexTangentKey = UE::Interchange::FAttributeKey(TEXT("HasVertexTangent"));
+	const UE::Interchange::FAttributeKey Macro_CustomHasSmoothGroupKey = UE::Interchange::FAttributeKey(TEXT("HasSmoothGroup"));
+	const UE::Interchange::FAttributeKey Macro_CustomHasVertexColorKey = UE::Interchange::FAttributeKey(TEXT("HasVertexColor"));
+	const UE::Interchange::FAttributeKey Macro_CustomUVCountKey = UE::Interchange::FAttributeKey(TEXT("UVCount"));
+
+	UE::Interchange::FNameAttributeArrayHelper SkeletonDependencies;
+	UE::Interchange::FNameAttributeArrayHelper MaterialDependencies;
+	UE::Interchange::FNameAttributeArrayHelper ShapeDependencies;
+};
diff --git a/Engine/Plugins/Runtime/Interchange/Source/Nodes/Public/InterchangeSceneNode.h b/Engine/Plugins/Runtime/Interchange/Source/Nodes/Public/InterchangeSceneNode.h
--- a/Engine/Plugins/Runtime/Interchange/Source/Nodes/Public/InterchangeSceneNode.h
+++ b/Engine/Plugins/Runtime/Interchange/Source/Nodes/Public/InterchangeSceneNode.h
@@ -21,6 +21,12 @@ namespace UE
 				static FString SceneNodeSpecializeType_BaseKey = TEXT("SceneNodeSpecializeType");
 				return SceneNodeSpecializeType_BaseKey;
 			}
+
+			static const FString& GetAssetDependenciesBaseKey()
+			{
+				static FString AssetDependencies_BaseKey = TEXT("__AssetDependenciesBaseKey__");
+				return AssetDependencies_BaseKey;
+			}
 		};
 
 	}//ns Interchange
@@ -36,6 +42,7 @@ public:
 	:UInterchangeBaseNode()
 	{
 		NodeSpecializeTypes.Initialize(Attributes, UE::Interchange::FSceneNodeStaticData::GetNodeSpecializeTypeBaseKey());
+		AssetDependencies.Initialize(Attributes, UE::Interchange::FSceneNodeStaticData::GetAssetDependenciesBaseKey());
 	}
 
 	/**
@@ -52,7 +59,7 @@ public:
 		FString KeyDisplayName = NodeAttributeKey.Key;
 		if (NodeAttributeKey.Key.Equals(UE::Interchange::FSceneNodeStaticData::GetNodeSpecializeTypeBaseKey()))
 		{
-			KeyDisplayName = TEXT("Specialized Type Count");
+			KeyDisplayName = TEXT("Specialized type count");
 			return KeyDisplayName;
 		}
 		else if (NodeAttributeKey.Key.StartsWith(UE::Interchange::FSceneNodeStaticData::GetNodeSpecializeTypeBaseKey()))
@@ -66,9 +73,43 @@ public:
 			}
 			return KeyDisplayName;
 		}
+		else if (NodeAttributeKey.Key.Equals(UE::Interchange::FSceneNodeStaticData::GetAssetDependenciesBaseKey()))
+		{
+			KeyDisplayName = TEXT("Asset dependencies count");
+			return KeyDisplayName;
+		}
+		else if (NodeAttributeKey.Key.StartsWith(UE::Interchange::FSceneNodeStaticData::GetAssetDependenciesBaseKey()))
+		{
+			KeyDisplayName = TEXT("Asset dependency index ");
+			const FString IndexKey = UE::Interchange::FNameAttributeArrayHelper::IndexKey();
+			int32 IndexPosition = NodeAttributeKey.Key.Find(IndexKey) + IndexKey.Len();
+			if (IndexPosition < NodeAttributeKey.Key.Len())
+			{
+				KeyDisplayName += NodeAttributeKey.Key.RightChop(IndexPosition);
+			}
+			return KeyDisplayName;
+		}
 		return Super::GetKeyDisplayName(NodeAttributeKey);
 	}
 
+	virtual FString GetAttributeCategory(const UE::Interchange::FAttributeKey& NodeAttributeKey) const override
+	{
+		if (NodeAttributeKey.Key.StartsWith(UE::Interchange::FSceneNodeStaticData::GetNodeSpecializeTypeBaseKey()))
+		{
+			return FString(TEXT("SpecializeType"));
+		}
+		else if (NodeAttributeKey.Key.StartsWith(UE::Interchange::FSceneNodeStaticData::GetAssetDependenciesBaseKey()))
+		{
+			return FString(TEXT("AssetDependencies"));
+		}
+		else if (NodeAttributeKey == Macro_CustomLocalTransformKey
+				 || NodeAttributeKey == Macro_CustomGlobalTransformKey)
+		{
+			return FString(TEXT("Scene"));
+		}
+		return Super::GetAttributeCategory(NodeAttributeKey);
+	}
+
 	virtual FGuid GetHash() const override
 	{
 		return Attributes->GetStorageHash();
@@ -120,6 +161,37 @@ public:
 		return NodeSpecializeTypes.RemoveName(SpecializedType);
 	}
 
+	/** Asset dependencies are the asset on which this node depend.*/
+	UFUNCTION(BlueprintCallable, Category = "Interchange | Node | Scene")
+	int32 GetAssetDependenciesCount() const
+	{
+		return AssetDependencies.GetCount();
+	}
+
+	UFUNCTION(BlueprintCallable, Category = "Interchange | Node | Scene")
+	void GetAssetDependency(const int32 Index, FString& OutAssetDependency) const
+	{
+		AssetDependencies.GetName(Index, OutAssetDependency);
+	}
+
+	UFUNCTION(BlueprintCallable, Category = "Interchange | Node | Scene")
+	void GetAssetDependencies(TArray<FString>& OutAssetDependencies) const
+	{
+		AssetDependencies.GetNames(OutAssetDependencies);
+	}
+
+	UFUNCTION(BlueprintCallable, Category = "Interchange | Node | Scene")
+	bool AddAssetDependency(const FString& AssetDependency)
+	{
+		return AssetDependencies.AddName(AssetDependency);
+	}
+
+	UFUNCTION(BlueprintCallable, Category = "Interchange | Node | Scene")
+	bool RemoveAssetDependency(const FString& AssetDependency)
+	{
+		return AssetDependencies.RemoveName(AssetDependency);
+	}
+
 	/** Return false if the Attribute was not set previously.*/
 	UFUNCTION(BlueprintCallable, Category = "Interchange | Node | Scene")
 	bool GetCustomLocalTransform(FTransform& AttributeValue) const
@@ -152,4 +224,5 @@ private:
 	const UE::Interchange::FAttributeKey Macro_CustomGlobalTransformKey = UE::Interchange::FAttributeKey(TEXT("GlobalTransform"));
 
 	UE::Interchange::FNameAttributeArrayHelper NodeSpecializeTypes;
+	UE::Interchange::FNameAttributeArrayHelper AssetDependencies;
 };
diff --git a/Engine/Plugins/Runtime/Interchange/Source/Nodes/Public/InterchangeSkeletalMeshFactoryNode.h b/Engine/Plugins/Runtime/Interchange/Source/Nodes/Public/InterchangeSkeletalMeshFactoryNode.h
new file mode 100644
--- /dev/null
+++ b/Engine/Plugins/Runtime/Interchange/Source/Nodes/Public/InterchangeSkeletalMeshFactoryNode.h
@@ -0,0 +1,193 @@
+// Copyright Epic Games, Inc. All Rights Reserved.
+
+#pragma once
+
+#include "CoreMinimal.h"
+#include "Nodes/InterchangeBaseNode.h"
+
+#if WITH_ENGINE
+#include "Engine/SkeletalMesh.h"
+#endif
+
+#include "InterchangeSkeletalMeshFactoryNode.generated.h"
+
+//Interchange namespace
+namespace UE
+{
+	namespace Interchange
+	{
+
+		struct FSkeletalMeshNodeStaticData : public FBaseNodeStaticData
+		{
+			static const FString& GetLodDependenciesBaseKey()
+			{
+				static FString LodDependencies_BaseKey = TEXT("Lod_Dependencies");
+				return LodDependencies_BaseKey;
+			}
+		};
+
+	}//ns Interchange
+}//ns UE
+
+UCLASS(BlueprintType)
+class INTERCHANGENODES_API UInterchangeSkeletalMeshFactoryNode : public UInterchangeBaseNode
+{
+	GENERATED_BODY()
+
+public:
+	UInterchangeSkeletalMeshFactoryNode()
+		:UInterchangeBaseNode()
+	{
+#if WITH_ENGINE
+		AssetClass = nullptr;
+#endif
+		LodDependencies.Initialize(Attributes, UE::Interchange::FSkeletalMeshNodeStaticData::GetLodDependenciesBaseKey());
+	}
+
+	/**
+	 * Initialize node data
+	 * @param: UniqueID - The uniqueId for this node
+	 * @param DisplayLabel - The name of the node
+	 * @param InAssetClass - The class the SkeletalMesh factory will create for this node.
+	 */
+	UFUNCTION(BlueprintCallable, Category = "Interchange | Node | SkeletalMesh")
+	void InitializeSkeletalMeshNode(const FString& UniqueID, const FString& DisplayLabel, const FString& InAssetClass)
+	{
+		bIsNodeClassInitialized = false;
+		InitializeNode(UniqueID, DisplayLabel, EInterchangeNodeContainerType::NodeContainerType_FactoryData);
+
+		FString OperationName = GetTypeName() + TEXT(".SetAssetClassName");
+		InterchangePrivateNodeBase::SetCustomAttribute<FString>(*Attributes, ClassNameAttributeKey, OperationName, InAssetClass);
+		FillAssetClassFromAttribute();
+	}
+
+	virtual void Serialize(FArchive& Ar) override
+	{
+		Super::Serialize(Ar);
+#if WITH_ENGINE
+		if (Ar.IsLoading())
+		{
+			//Make sure the class is properly set when we compile with engine, this will set the
+			//bIsNodeClassInitialized to true.
+			SetNodeClassFromClassAttribute();
+		}
+#endif //#if WITH_ENGINE
+	}
+
+	/**
+	 * Return the node type name of the class, we use this when reporting error
+	 */
+	virtual FString GetTypeName() const override
+	{
+		const FString TypeName = TEXT("SkeletalMeshNode");
+		return TypeName;
+	}
+
+	virtual FString GetKeyDisplayName(const UE::Interchange::FAttributeKey& NodeAttributeKey) const override
+	{
+		FString KeyDisplayName = NodeAttributeKey.Key;
+		if (NodeAttributeKey.Key.Equals(UE::Interchange::FSkeletalMeshNodeStaticData::GetLodDependenciesBaseKey()))
+		{
+			KeyDisplayName = TEXT("LOD Dependencies Count");
+			return KeyDisplayName;
+		}
+		else if (NodeAttributeKey.Key.StartsWith(UE::Interchange::FSkeletalMeshNodeStaticData::GetLodDependenciesBaseKey()))
+		{
+			KeyDisplayName = TEXT("LOD Dependencies Index ");
+			const FString IndexKey = UE::Interchange::FNameAttributeArrayHelper::IndexKey();
+			int32 IndexPosition = NodeAttributeKey.Key.Find(IndexKey) + IndexKey.Len();
+			if (IndexPosition < NodeAttributeKey.Key.Len())
+			{
+				KeyDisplayName += NodeAttributeKey.Key.RightChop(IndexPosition);
+			}
+			return KeyDisplayName;
+		}
+		return Super::GetKeyDisplayName(NodeAttributeKey);
+	}
+
+	/** Get the class this node want to create */
+	UFUNCTION(BlueprintCallable, Category = "Interchange | Node | SkeletalMesh")
+	virtual class UClass* GetAssetClass() const override
+	{
+		ensure(bIsNodeClassInitialized);
+#if WITH_ENGINE
+		return AssetClass.Get() != nullptr ? AssetClass.Get() : USkeletalMesh::StaticClass();
+#else
+		return nullptr;
+#endif
+	}
+
+	virtual FGuid GetHash() const override
+	{
+		return Attributes->GetStorageHash();
+	}
+
+public:
+	/** Return The number of LOD this skeletalmesh has.*/
+	UFUNCTION(BlueprintCallable, Category = "Interchange | Node | SkeletalMesh")
+	int32 GetLodDataCount() const
+	{
+		return LodDependencies.GetCount();
+	}
+
+	UFUNCTION(BlueprintCallable, Category = "Interchange | Node | SkeletalMesh")
+	void GetLodDataUniqueIds(TArray<FString>& OutLodDataUniqueIds) const
+	{
+		LodDependencies.GetNames(OutLodDataUniqueIds);
+	}
+
+	UFUNCTION(BlueprintCallable, Category = "Interchange | Node | SkeletalMesh")
+	bool AddLodDataUniqueId(const FString& LodDataUniqueId)
+	{
+		return LodDependencies.AddName(LodDataUniqueId);
+	}
+
+	UFUNCTION(BlueprintCallable, Category = "Interchange | Node | SkeletalMesh")
+	bool RemoveLodDataUniqueId(const FString& LodDataUniqueId)
+	{
+		return LodDependencies.RemoveName(LodDataUniqueId);
+	}
+
+private:
+
+	void FillAssetClassFromAttribute()
+	{
+#if WITH_ENGINE
+		FString OperationName = GetTypeName() + TEXT(".GetAssetClassName");
+		FString ClassName;
+		InterchangePrivateNodeBase::GetCustomAttribute<FString>(*Attributes, ClassNameAttributeKey, OperationName, ClassName);
+		if (ClassName.Equals(USkeletalMesh::StaticClass()->GetName()))
+		{
+			AssetClass = USkeletalMesh::StaticClass();
+			bIsNodeClassInitialized = true;
+		}
+#endif
+	}
+
+	bool SetNodeClassFromClassAttribute()
+	{
+		if (!bIsNodeClassInitialized)
+		{
+			FillAssetClassFromAttribute();
+		}
+		return bIsNodeClassInitialized;
+	}
+
+	bool IsEditorOnlyDataDefined()
+	{
+#if WITH_EDITORONLY_DATA
+		return true;
+#else
+		return false;
+#endif
+	}
+
+	const UE::Interchange::FAttributeKey ClassNameAttributeKey = UE::Interchange::FBaseNodeStaticData::ClassTypeAttributeKey();
+
+	UE::Interchange::FNameAttributeArrayHelper LodDependencies;
+protected:
+#if WITH_ENGINE
+	TSubclassOf<USkeletalMesh> AssetClass = nullptr;
+#endif
+	bool bIsNodeClassInitialized = false;
+};
diff --git a/Engine/Plugins/Runtime/Interchange/Source/Nodes/Public/InterchangeSkeletalMeshNode.h b/Engine/Plugins/Runtime/Interchange/Source/Nodes/Public/InterchangeSkeletalMeshNode.h
deleted file mode 100644
diff --git a/Engine/Plugins/Runtime/Interchange/Source/Nodes/Public/InterchangeSkeletonFactoryNode.h b/Engine/Plugins/Runtime/Interchange/Source/Nodes/Public/InterchangeSkeletonFactoryNode.h
new file mode 100644
--- /dev/null
+++ b/Engine/Plugins/Runtime/Interchange/Source/Nodes/Public/InterchangeSkeletonFactoryNode.h
@@ -0,0 +1,141 @@
+// Copyright Epic Games, Inc. All Rights Reserved.
+
+#pragma once
+
+#include "CoreMinimal.h"
+#include "Nodes/InterchangeBaseNode.h"
+
+#if WITH_ENGINE
+#include "Animation/Skeleton.h"
+#endif
+
+#include "InterchangeSkeletonFactoryNode.generated.h"
+
+UCLASS(BlueprintType)
+class INTERCHANGENODES_API UInterchangeSkeletonFactoryNode : public UInterchangeBaseNode
+{
+	GENERATED_BODY()
+
+public:
+	UInterchangeSkeletonFactoryNode()
+	:UInterchangeBaseNode()
+	{
+#if WITH_ENGINE
+		AssetClass = nullptr;
+#endif
+	}
+
+	/**
+	 * Initialize node data
+	 * @param: UniqueID - The uniqueId for this node
+	 * @param DisplayLabel - The name of the node
+	 * @param InAssetClass - The class the Skeleton factory will create for this node.
+	 */
+	UFUNCTION(BlueprintCallable, Category = "Interchange | Node | Skeleton")
+	void InitializeSkeletonNode(const FString& UniqueID, const FString& DisplayLabel, const FString& InAssetClass)
+	{
+		bIsNodeClassInitialized = false;
+		InitializeNode(UniqueID, DisplayLabel, EInterchangeNodeContainerType::NodeContainerType_FactoryData);
+		FString OperationName = GetTypeName() + TEXT(".SetAssetClassName");
+		InterchangePrivateNodeBase::SetCustomAttribute<FString>(*Attributes, ClassNameAttributeKey, OperationName, InAssetClass);
+		FillAssetClassFromAttribute();
+	}
+
+	virtual void Serialize(FArchive& Ar) override
+	{
+		Super::Serialize(Ar);
+#if WITH_ENGINE
+		if(Ar.IsLoading())
+		{
+			//Make sure the class is properly set when we compile with engine, this will set the
+			//bIsNodeClassInitialized to true.
+			SetNodeClassFromClassAttribute();
+		}
+#endif //#if WITH_ENGINE
+	}
+
+	/**
+	 * Return the node type name of the class, we use this when reporting error
+	 */
+	virtual FString GetTypeName() const override
+	{
+		const FString TypeName = TEXT("SkeletonNode");
+		return TypeName;
+	}
+
+	/** Get the class this node want to create */
+	UFUNCTION(BlueprintCallable, Category = "Interchange | Node | Skeleton")
+	virtual class UClass* GetAssetClass() const override
+	{
+		ensure(bIsNodeClassInitialized);
+#if WITH_ENGINE
+		return AssetClass.Get() != nullptr ? AssetClass.Get() : USkeleton::StaticClass();
+#else
+		return nullptr;
+#endif
+	}
+
+	virtual FGuid GetHash() const override
+	{
+		return Attributes->GetStorageHash();
+	}
+
+public:
+	/** Return false if the Attribute was not set previously.*/
+	UFUNCTION(BlueprintCallable, Category = "Interchange | Node | Skeleton")
+	bool GetCustomRootJointID(FString& AttributeValue) const
+	{
+		IMPLEMENT_NODE_ATTRIBUTE_GETTER(RootJointID, FString);
+	}
+
+	UFUNCTION(BlueprintCallable, Category = "Interchange | Node | Skeleton")
+	bool SetCustomRootJointID(const FString& AttributeValue)
+	{
+		IMPLEMENT_NODE_ATTRIBUTE_SETTER_NODELEGATE(RootJointID, FString)
+	}
+
+private:
+
+	void FillAssetClassFromAttribute()
+	{
+#if WITH_ENGINE
+		FString OperationName = GetTypeName() + TEXT(".GetAssetClassName");
+		FString ClassName;
+		InterchangePrivateNodeBase::GetCustomAttribute<FString>(*Attributes, ClassNameAttributeKey, OperationName, ClassName);
+		if (ClassName.Equals(USkeleton::StaticClass()->GetName()))
+		{
+			AssetClass = USkeleton::StaticClass();
+			bIsNodeClassInitialized = true;
+		}
+#endif
+	}
+
+	bool SetNodeClassFromClassAttribute()
+	{
+		if (!bIsNodeClassInitialized)
+		{
+			FillAssetClassFromAttribute();
+		}
+		return bIsNodeClassInitialized;
+	}
+
+	bool IsEditorOnlyDataDefined()
+	{
+#if WITH_EDITORONLY_DATA
+		return true;
+#else
+		return false;
+#endif
+	}
+
+	const UE::Interchange::FAttributeKey ClassNameAttributeKey = UE::Interchange::FBaseNodeStaticData::ClassTypeAttributeKey();
+
+	//Skeleton
+	const UE::Interchange::FAttributeKey Macro_CustomRootJointIDKey = UE::Interchange::FAttributeKey(TEXT("RootJointID"));
+	
+protected:
+#if WITH_ENGINE
+	TSubclassOf<USkeleton> AssetClass = nullptr;
+#endif
+	bool bIsNodeClassInitialized = false;
+};
diff --git a/Engine/Plugins/Runtime/Interchange/Source/Nodes/Public/InterchangeSkeletonNode.h b/Engine/Plugins/Runtime/Interchange/Source/Nodes/Public/InterchangeSkeletonNode.h
deleted file mode 100644
diff --git a/Engine/Plugins/Runtime/Interchange/Source/Nodes/Public/InterchangeTextureFactoryNode.h b/Engine/Plugins/Runtime/Interchange/Source/Nodes/Public/InterchangeTextureFactoryNode.h
new file mode 100644
--- /dev/null
+++ b/Engine/Plugins/Runtime/Interchange/Source/Nodes/Public/InterchangeTextureFactoryNode.h
@@ -0,0 +1,902 @@
+// Copyright Epic Games, Inc. All Rights Reserved.
+
+#pragma once
+
+#include "CoreMinimal.h"
+#include "Nodes/InterchangeBaseNode.h"
+
+#if WITH_ENGINE
+#include "Engine/Texture.h"
+#include "Engine/Texture2D.h"
+#endif
+
+#include "InterchangeTextureFactoryNode.generated.h"
+
+
+UCLASS(BlueprintType)
+class INTERCHANGENODES_API UInterchangeTextureFactoryNode : public UInterchangeBaseNode
+{
+	GENERATED_BODY()
+
+public:
+	UInterchangeTextureFactoryNode()
+		:UInterchangeBaseNode()
+	{
+#if WITH_ENGINE
+		AssetClass = nullptr;
+#endif
+	}
+
+	/**
+	 * Initialize node data
+	 * @param: UniqueID - The uniqueId for this node
+	 * @param DisplayLabel - The name of the node
+	 * @param InAssetClass - The class the texture factory will create for this node.
+	 */
+	UFUNCTION(BlueprintCallable, Category = "Interchange | Node | Texture")
+	void InitializeTextureNode(const FString& UniqueID, const FString& DisplayLabel, const FString& InAssetClass, const FString& InAssetName)
+	{
+		//Initialize
+		bIsTextureNodeClassInitialized = false;
+		InitializeNode(UniqueID, DisplayLabel, EInterchangeNodeContainerType::NodeContainerType_FactoryData);
+		
+		//Set the class
+		FString OperationName = GetTypeName() + TEXT(".SetAssetClassName");
+		InterchangePrivateNodeBase::SetCustomAttribute<FString>(*Attributes, ClassNameAttributeKey, OperationName, InAssetClass);
+		FillAssetClassFromAttribute();
+		
+		//Set the asset name
+		OperationName = GetTypeName() + TEXT(".SetAssetName");
+		InterchangePrivateNodeBase::SetCustomAttribute<FString>(*Attributes, AssetNameKey, OperationName, InAssetName);
+	}
+
+	virtual void Serialize(FArchive& Ar) override
+	{
+		Super::Serialize(Ar);
+#if WITH_ENGINE
+		if (Ar.IsLoading())
+		{
+			//Make sure the class is properly set when we compile with engine, this will set the
+			//bIsTextureNodeClassInitialized to true.
+			SetTextureNodeClassFromClassAttribute();
+		}
+#endif //#if WITH_ENGINE
+	}
+
+	/**
+	 * Return the node type name of the class, we use this when reporting error
+	 */
+	virtual FString GetTypeName() const override
+	{
+		const FString TypeName = TEXT("TextureFactoryNode");
+		return TypeName;
+	}
+
+	/** Get the class this node want to create */
+	UFUNCTION(BlueprintCallable, Category = "Interchange | Node | Texture")
+	virtual class UClass* GetAssetClass() const override
+	{
+		ensure(bIsTextureNodeClassInitialized);
+#if WITH_ENGINE
+		return AssetClass.Get() != nullptr ? AssetClass.Get() : UTexture::StaticClass();
+#else
+		return nullptr;
+#endif
+	}
+
+	virtual FGuid GetHash() const override
+	{
+		return Attributes->GetStorageHash();
+	}
+
+	static FString GetTextureFactoryNodeUidFromTextureNodeUid(const FString& TranslatedNodeUid)
+	{
+		FString NewUid = TEXT("Factory_") + TranslatedNodeUid;
+		return NewUid;
+	}
+
+	/**
+	 * Get the translated texture node unique ID.
+	 */
+	UFUNCTION(BlueprintCallable, Category = "Interchange | Node | Texture")
+	bool GetCustomTranslatedTextureNodeUid(FString& AttributeValue) const
+	{
+		IMPLEMENT_NODE_ATTRIBUTE_GETTER(TranslatedTextureNodeUid, FString);
+	}
+
+	/**
+	 * Set the translated texture node unique ID. This is the reference to the node that was create by the translator and this node is needed to get the texture payload.
+	 */
+	UFUNCTION(BlueprintCallable, Category = "Interchange | Node | Texture")
+	bool SetCustomTranslatedTextureNodeUid(const FString& AttributeValue)
+	{
+		IMPLEMENT_NODE_ATTRIBUTE_SETTER_NODELEGATE(TranslatedTextureNodeUid, FString);
+	}
+
+public:
+	/** Return false if the Attribute was not set previously.*/
+	UFUNCTION(BlueprintCallable, Category = "Interchange | Node | Texture")
+	bool GetCustomAdjustBrightness(float& AttributeValue) const
+	{
+		IMPLEMENT_NODE_ATTRIBUTE_GETTER(AdjustBrightness, float);
+	}
+
+	UFUNCTION(BlueprintCallable, Category = "Interchange | Node | Texture")
+	bool SetCustomAdjustBrightness(const float& AttributeValue, bool bAddApplyDelegate = true)
+	{
+#if WITH_EDITORONLY_DATA
+		IMPLEMENT_NODE_ATTRIBUTE_SETTER(UInterchangeTextureFactoryNode, AdjustBrightness, float, UTexture)
+#else
+		IMPLEMENT_NODE_ATTRIBUTE_SETTER_NODELEGATE(AdjustBrightness, float)
+#endif
+	}
+
+	/** Return false if the Attribute was not set previously.*/
+	UFUNCTION(BlueprintCallable, Category = "Interchange | Node | Texture")
+	bool GetCustomAdjustBrightnessCurve(float& AttributeValue) const
+	{
+		IMPLEMENT_NODE_ATTRIBUTE_GETTER(AdjustBrightnessCurve, float);
+	}
+
+	UFUNCTION(BlueprintCallable, Category = "Interchange | Node | Texture")
+	bool SetCustomAdjustBrightnessCurve(const float& AttributeValue, bool bAddApplyDelegate = true)
+	{
+#if WITH_EDITORONLY_DATA
+		IMPLEMENT_NODE_ATTRIBUTE_SETTER(UInterchangeTextureFactoryNode, AdjustBrightnessCurve, float, UTexture)
+#else
+		IMPLEMENT_NODE_ATTRIBUTE_SETTER_NODELEGATE(AdjustBrightnessCurve, float)
+#endif
+	}
+
+	/** Return false if the Attribute was not set previously.*/
+	UFUNCTION(BlueprintCallable, Category = "Interchange | Node | Texture")
+	bool GetCustomAdjustVibrance(float& AttributeValue) const
+	{
+		IMPLEMENT_NODE_ATTRIBUTE_GETTER(AdjustVibrance, float);
+	}
+
+	UFUNCTION(BlueprintCallable, Category = "Interchange | Node | Texture")
+	bool SetCustomAdjustVibrance(const float& AttributeValue, bool bAddApplyDelegate = true)
+	{
+#if WITH_EDITORONLY_DATA
+		IMPLEMENT_NODE_ATTRIBUTE_SETTER(UInterchangeTextureFactoryNode, AdjustVibrance, float, UTexture)
+#else
+		IMPLEMENT_NODE_ATTRIBUTE_SETTER_NODELEGATE(AdjustVibrance, float)
+#endif
+	}
+
+	/** Return false if the Attribute was not set previously.*/
+	UFUNCTION(BlueprintCallable, Category = "Interchange | Node | Texture")
+	bool GetCustomAdjustSaturation(float& AttributeValue) const
+	{
+		IMPLEMENT_NODE_ATTRIBUTE_GETTER(AdjustSaturation, float);
+	}
+
+	UFUNCTION(BlueprintCallable, Category = "Interchange | Node | Texture")
+	bool SetCustomAdjustSaturation(const float& AttributeValue, bool bAddApplyDelegate = true)
+	{
+#if WITH_EDITORONLY_DATA
+		IMPLEMENT_NODE_ATTRIBUTE_SETTER(UInterchangeTextureFactoryNode, AdjustSaturation, float, UTexture)
+#else
+		IMPLEMENT_NODE_ATTRIBUTE_SETTER_NODELEGATE(AdjustSaturation, float)
+#endif
+	}
+
+	/** Return false if the Attribute was not set previously.*/
+	UFUNCTION(BlueprintCallable, Category = "Interchange | Node | Texture")
+	bool GetCustomAdjustRGBCurve(float& AttributeValue) const
+	{
+		IMPLEMENT_NODE_ATTRIBUTE_GETTER(AdjustRGBCurve, float);
+	}
+
+	UFUNCTION(BlueprintCallable, Category = "Interchange | Node | Texture")
+	bool SetCustomAdjustRGBCurve(const float& AttributeValue, bool bAddApplyDelegate = true)
+	{
+#if WITH_EDITORONLY_DATA
+		IMPLEMENT_NODE_ATTRIBUTE_SETTER(UInterchangeTextureFactoryNode, AdjustRGBCurve, float, UTexture)
+#else
+		IMPLEMENT_NODE_ATTRIBUTE_SETTER_NODELEGATE(AdjustRGBCurve, float)
+#endif
+	}
+
+	/** Return false if the Attribute was not set previously.*/
+	UFUNCTION(BlueprintCallable, Category = "Interchange | Node | Texture")
+	bool GetCustomAdjustHue(float& AttributeValue) const
+	{
+		IMPLEMENT_NODE_ATTRIBUTE_GETTER(AdjustHue, float);
+	}
+
+	UFUNCTION(BlueprintCallable, Category = "Interchange | Node | Texture")
+	bool SetCustomAdjustHue(const float& AttributeValue, bool bAddApplyDelegate = true)
+	{
+#if WITH_EDITORONLY_DATA
+		IMPLEMENT_NODE_ATTRIBUTE_SETTER(UInterchangeTextureFactoryNode, AdjustHue, float, UTexture)
+#else
+		IMPLEMENT_NODE_ATTRIBUTE_SETTER_NODELEGATE(AdjustHue, float)
+#endif
+	}
+
+	/** Return false if the Attribute was not set previously.*/
+	UFUNCTION(BlueprintCallable, Category = "Interchange | Node | Texture")
+	bool GetCustomAdjustMinAlpha(float& AttributeValue) const
+	{
+		IMPLEMENT_NODE_ATTRIBUTE_GETTER(AdjustMinAlpha, float);
+	}
+
+	UFUNCTION(BlueprintCallable, Category = "Interchange | Node | Texture")
+	bool SetCustomAdjustMinAlpha(const float& AttributeValue, bool bAddApplyDelegate = true)
+	{
+#if WITH_EDITORONLY_DATA
+		IMPLEMENT_NODE_ATTRIBUTE_SETTER(UInterchangeTextureFactoryNode, AdjustMinAlpha, float, UTexture)
+#else
+		IMPLEMENT_NODE_ATTRIBUTE_SETTER_NODELEGATE(AdjustMinAlpha, float)
+#endif
+	}
+
+	/** Return false if the Attribute was not set previously.*/
+	UFUNCTION(BlueprintCallable, Category = "Interchange | Node | Texture")
+	bool GetCustomAdjustMaxAlpha(float& AttributeValue) const
+	{
+		IMPLEMENT_NODE_ATTRIBUTE_GETTER(AdjustMaxAlpha, float);
+	}
+
+	UFUNCTION(BlueprintCallable, Category = "Interchange | Node | Texture")
+	bool SetCustomAdjustMaxAlpha(const float& AttributeValue, bool bAddApplyDelegate = true)
+	{
+#if WITH_EDITORONLY_DATA
+		IMPLEMENT_NODE_ATTRIBUTE_SETTER(UInterchangeTextureFactoryNode, AdjustMaxAlpha, float, UTexture)
+#else
+		IMPLEMENT_NODE_ATTRIBUTE_SETTER_NODELEGATE(AdjustMaxAlpha, float)
+#endif
+	}
+
+	/** Return false if the Attribute was not set previously.*/
+	UFUNCTION(BlueprintCallable, Category = "Interchange | Node | Texture")
+	bool GetCustombChromaKeyTexture(bool& AttributeValue) const
+	{
+		IMPLEMENT_NODE_ATTRIBUTE_GETTER(bChromaKeyTexture, bool);
+	}
+
+	UFUNCTION(BlueprintCallable, Category = "Interchange | Node | Texture")
+	bool SetCustombChromaKeyTexture(const bool& AttributeValue, bool bAddApplyDelegate = true)
+	{
+#if WITH_EDITORONLY_DATA
+		IMPLEMENT_NODE_ATTRIBUTE_SETTER(UInterchangeTextureFactoryNode, bChromaKeyTexture, bool, UTexture)
+#else
+		IMPLEMENT_NODE_ATTRIBUTE_SETTER_NODELEGATE(bChromaKeyTexture, bool)
+#endif
+	}
+
+	/** Return false if the Attribute was not set previously.*/
+	UFUNCTION(BlueprintCallable, Category = "Interchange | Node | Texture")
+	bool GetCustomChromaKeyThreshold(float& AttributeValue) const
+	{
+		IMPLEMENT_NODE_ATTRIBUTE_GETTER(ChromaKeyThreshold, float);
+	}
+
+	UFUNCTION(BlueprintCallable, Category = "Interchange | Node | Texture")
+	bool SetCustomChromaKeyThreshold(const float& AttributeValue, bool bAddApplyDelegate = true)
+	{
+#if WITH_EDITORONLY_DATA
+		IMPLEMENT_NODE_ATTRIBUTE_SETTER(UInterchangeTextureFactoryNode, ChromaKeyThreshold, float, UTexture)
+#else
+		IMPLEMENT_NODE_ATTRIBUTE_SETTER_NODELEGATE(ChromaKeyThreshold, float)
+#endif
+	}
+
+	/** Return false if the Attribute was not set previously.*/
+	UFUNCTION(BlueprintCallable, Category = "Interchange | Node | Texture")
+	bool GetCustomChromaKeyColor(FColor& AttributeValue) const
+	{
+		IMPLEMENT_NODE_ATTRIBUTE_GETTER(ChromaKeyColor, FColor);
+	}
+
+	UFUNCTION(BlueprintCallable, Category = "Interchange | Node | Texture")
+	bool SetCustomChromaKeyColor(const FColor& AttributeValue, bool bAddApplyDelegate = true)
+	{
+#if WITH_EDITORONLY_DATA
+		IMPLEMENT_NODE_ATTRIBUTE_SETTER(UInterchangeTextureFactoryNode, ChromaKeyColor, FColor, UTexture)
+#else
+		IMPLEMENT_NODE_ATTRIBUTE_SETTER_NODELEGATE(ChromaKeyColor, FColor)
+#endif
+	}
+
+	//////////////////////////////////////////////////////////////////////////
+	//Texture Compression Begin
+
+	/** Return false if the Attribute was not set previously.*/
+	UFUNCTION(BlueprintCallable, Category = "Interchange | Node | Texture")
+	bool GetCustomCompressionNoAlpha(bool& AttributeValue) const
+	{
+		IMPLEMENT_NODE_ATTRIBUTE_GETTER(CompressionNoAlpha, bool);
+	}
+
+	UFUNCTION(BlueprintCallable, Category = "Interchange | Node | Texture")
+	bool SetCustomCompressionNoAlpha(const bool& AttributeValue, bool bAddApplyDelegate = true)
+	{
+#if WITH_EDITORONLY_DATA
+		IMPLEMENT_NODE_ATTRIBUTE_SETTER(UInterchangeTextureFactoryNode, CompressionNoAlpha, bool, UTexture)
+#else
+		IMPLEMENT_NODE_ATTRIBUTE_SETTER_NODELEGATE(CompressionNoAlpha, bool)
+#endif
+	}
+
+	/** Return false if the Attribute was not set previously.*/
+	UFUNCTION(BlueprintCallable, Category = "Interchange | Node | Texture")
+	bool GetCustomDeferCompression(bool& AttributeValue) const
+	{
+		IMPLEMENT_NODE_ATTRIBUTE_GETTER(DeferCompression, bool);
+	}
+
+	UFUNCTION(BlueprintCallable, Category = "Interchange | Node | Texture")
+	bool SetCustomDeferCompression(const bool& AttributeValue, bool bAddApplyDelegate = true)
+	{
+#if WITH_EDITORONLY_DATA
+		IMPLEMENT_NODE_ATTRIBUTE_SETTER(UInterchangeTextureFactoryNode, DeferCompression, bool, UTexture)
+#else
+		IMPLEMENT_NODE_ATTRIBUTE_SETTER_NODELEGATE(DeferCompression, bool)
+#endif
+	}
+
+	/** Return false if the Attribute was not set previously.*/
+	UFUNCTION(BlueprintCallable, Category = "Interchange | Node | Texture")
+	bool GetCustomLossyCompressionAmount(uint8& AttributeValue) const
+	{
+		IMPLEMENT_NODE_ATTRIBUTE_GETTER(LossyCompressionAmount, uint8);
+	}
+
+	UFUNCTION(BlueprintCallable, Category = "Interchange | Node | Texture")
+	bool SetCustomLossyCompressionAmount(const uint8& AttributeValue, bool bAddApplyDelegate = true)
+	{
+#if WITH_EDITORONLY_DATA
+		IMPLEMENT_NODE_ATTRIBUTE_SETTER(UInterchangeTextureFactoryNode, LossyCompressionAmount, uint8, UTexture)
+#else
+		IMPLEMENT_NODE_ATTRIBUTE_SETTER_NODELEGATE(LossyCompressionAmount, uint8)
+#endif
+	}
+
+	/** Return false if the Attribute was not set previously.*/
+	UFUNCTION(BlueprintCallable, Category = "Interchange | Node | Texture")
+	bool GetCustomMaxTextureSize(int32& AttributeValue) const
+	{
+		IMPLEMENT_NODE_ATTRIBUTE_GETTER(MaxTextureSize, int32);
+	}
+
+	UFUNCTION(BlueprintCallable, Category = "Interchange | Node | Texture")
+	bool SetCustomMaxTextureSize(const int32& AttributeValue, bool bAddApplyDelegate = true)
+	{
+#if WITH_EDITORONLY_DATA
+		IMPLEMENT_NODE_ATTRIBUTE_SETTER(UInterchangeTextureFactoryNode, MaxTextureSize, int32, UTexture)
+#else
+		IMPLEMENT_NODE_ATTRIBUTE_SETTER_NODELEGATE(MaxTextureSize, int32)
+#endif
+	}
+
+	/** Return false if the Attribute was not set previously.*/
+	UFUNCTION(BlueprintCallable, Category = "Interchange | Node | Texture")
+	bool GetCustomCompressionQuality(uint8& AttributeValue) const
+	{
+		IMPLEMENT_NODE_ATTRIBUTE_GETTER(CompressionQuality, uint8);
+	}
+
+	UFUNCTION(BlueprintCallable, Category = "Interchange | Node | Texture")
+	bool SetCustomCompressionQuality(const uint8& AttributeValue, bool bAddApplyDelegate = true)
+	{
+#if WITH_EDITORONLY_DATA
+		IMPLEMENT_NODE_ATTRIBUTE_SETTER(UInterchangeTextureFactoryNode, CompressionQuality, uint8, UTexture)
+#else
+		IMPLEMENT_NODE_ATTRIBUTE_SETTER_NODELEGATE(CompressionQuality, uint8)
+#endif
+	}
+
+	/** Return false if the Attribute was not set previously.*/
+	UFUNCTION(BlueprintCallable, Category = "Interchange | Node | Texture")
+	bool GetCustomCompressionSettings(uint8& AttributeValue) const
+	{
+		IMPLEMENT_NODE_ATTRIBUTE_GETTER(CompressionSettings, uint8);
+	}
+
+	UFUNCTION(BlueprintCallable, Category = "Interchange | Node | Texture")
+	bool SetCustomCompressionSettings(const uint8& AttributeValue, bool bAddApplyDelegate = true)
+	{
+		IMPLEMENT_NODE_ATTRIBUTE_SETTER(UInterchangeTextureFactoryNode, CompressionSettings, uint8, UTexture)
+	}
+
+	//////////////////////////////////////////////////////////////////////////
+	//Texture general
+
+	/** Return false if the Attribute was not set previously.*/
+	UFUNCTION(BlueprintCallable, Category = "Interchange | Node | Texture")
+	bool GetCustombDitherMipMapAlpha(bool& AttributeValue) const
+	{
+		IMPLEMENT_NODE_ATTRIBUTE_GETTER(bDitherMipMapAlpha, bool);
+	}
+
+	UFUNCTION(BlueprintCallable, Category = "Interchange | Node | Texture")
+	bool SetCustombDitherMipMapAlpha(const bool& AttributeValue, bool bAddApplyDelegate = true)
+	{
+#if WITH_EDITORONLY_DATA
+		IMPLEMENT_NODE_ATTRIBUTE_SETTER(UInterchangeTextureFactoryNode, bDitherMipMapAlpha, bool, UTexture)
+#else
+		IMPLEMENT_NODE_ATTRIBUTE_SETTER_NODELEGATE(bDitherMipMapAlpha, bool)
+#endif
+	}
+
+	/** Return false if the Attribute was not set previously.*/
+	UFUNCTION(BlueprintCallable, Category = "Interchange | Node | Texture")
+	bool GetCustomAlphaCoverageThresholds(FVector4& AttributeValue) const
+	{
+		IMPLEMENT_NODE_ATTRIBUTE_GETTER(AlphaCoverageThresholds, FVector4);
+	}
+
+	UFUNCTION(BlueprintCallable, Category = "Interchange | Node | Texture")
+	bool SetCustomAlphaCoverageThresholds(const FVector4& AttributeValue, bool bAddApplyDelegate = true)
+	{
+#if WITH_EDITORONLY_DATA
+		IMPLEMENT_NODE_ATTRIBUTE_SETTER(UInterchangeTextureFactoryNode, AlphaCoverageThresholds, FVector4, UTexture)
+#else
+		IMPLEMENT_NODE_ATTRIBUTE_SETTER_NODELEGATE(AlphaCoverageThresholds, FVector4)
+#endif
+	}
+
+	/** Return false if the Attribute was not set previously.*/
+	UFUNCTION(BlueprintCallable, Category = "Interchange | Node | Texture")
+	bool GetCustombFlipGreenChannel(bool& AttributeValue) const
+	{
+		IMPLEMENT_NODE_ATTRIBUTE_GETTER(bFlipGreenChannel, bool);
+	}
+
+	UFUNCTION(BlueprintCallable, Category = "Interchange | Node | Texture")
+	bool SetCustombFlipGreenChannel(const bool& AttributeValue, bool bAddApplyDelegate = true)
+	{
+#if WITH_EDITORONLY_DATA
+		IMPLEMENT_NODE_ATTRIBUTE_SETTER(UInterchangeTextureFactoryNode, bFlipGreenChannel, bool, UTexture)
+#else
+		IMPLEMENT_NODE_ATTRIBUTE_SETTER_NODELEGATE(bFlipGreenChannel, bool)
+#endif
+	}
+
+	/** Return false if the Attribute was not set previously.*/
+	UFUNCTION(BlueprintCallable, Category = "Interchange | Node | Texture")
+	bool GetCustomPowerOfTwoMode(uint8& AttributeValue) const
+	{
+		IMPLEMENT_NODE_ATTRIBUTE_GETTER(PowerOfTwoMode, uint8);
+	}
+
+	UFUNCTION(BlueprintCallable, Category = "Interchange | Node | Texture")
+	bool SetCustomPowerOfTwoMode(const uint8& AttributeValue, bool bAddApplyDelegate = true)
+	{
+#if WITH_EDITORONLY_DATA
+		IMPLEMENT_NODE_ATTRIBUTE_SETTER(UInterchangeTextureFactoryNode, PowerOfTwoMode, uint8, UTexture)
+#else
+		IMPLEMENT_NODE_ATTRIBUTE_SETTER_NODELEGATE(PowerOfTwoMode, uint8)
+#endif
+	}
+
+	/** Return false if the Attribute was not set previously.*/
+	UFUNCTION(BlueprintCallable, Category = "Interchange | Node | Texture")
+	bool GetCustomPaddingColor(FColor& AttributeValue) const
+	{
+		IMPLEMENT_NODE_ATTRIBUTE_GETTER(PaddingColor, FColor);
+	}
+
+	UFUNCTION(BlueprintCallable, Category = "Interchange | Node | Texture")
+	bool SetCustomPaddingColor(const FColor& AttributeValue, bool bAddApplyDelegate = true)
+	{
+#if WITH_EDITORONLY_DATA
+		IMPLEMENT_NODE_ATTRIBUTE_SETTER(UInterchangeTextureFactoryNode, PaddingColor, FColor, UTexture)
+#else
+		IMPLEMENT_NODE_ATTRIBUTE_SETTER_NODELEGATE(PaddingColor, FColor)
+#endif
+	}
+
+	/** Return false if the Attribute was not set previously.*/
+	UFUNCTION(BlueprintCallable, Category = "Interchange | Node | Texture")
+	bool GetCustomFilter(uint8& AttributeValue) const
+	{
+		IMPLEMENT_NODE_ATTRIBUTE_GETTER(Filter, uint8);
+	}
+
+	UFUNCTION(BlueprintCallable, Category = "Interchange | Node | Texture")
+	bool SetCustomFilter(const uint8& AttributeValue, bool bAddApplyDelegate = true)
+	{
+		IMPLEMENT_NODE_ATTRIBUTE_SETTER(UInterchangeTextureFactoryNode, Filter, uint8, UTexture)
+	}
+
+	/** Return false if the Attribute was not set previously.*/
+	UFUNCTION(BlueprintCallable, Category = "Interchange | Node | Texture")
+	bool GetCustomMipLoadOptions(uint8& AttributeValue) const
+	{
+		IMPLEMENT_NODE_ATTRIBUTE_GETTER(MipLoadOptions, uint8);
+	}
+
+	UFUNCTION(BlueprintCallable, Category = "Interchange | Node | Texture")
+	bool SetCustomMipLoadOptions(const uint8& AttributeValue, bool bAddApplyDelegate = true)
+	{
+		IMPLEMENT_NODE_ATTRIBUTE_SETTER(UInterchangeTextureFactoryNode, MipLoadOptions, uint8, UTexture)
+	}
+
+	/** Return false if the Attribute was not set previously.*/
+	UFUNCTION(BlueprintCallable, Category = "Interchange | Node | Texture")
+	bool GetCustomSRGB(bool& AttributeValue) const
+	{
+		IMPLEMENT_NODE_ATTRIBUTE_GETTER(SRGB, bool);
+	}
+
+	UFUNCTION(BlueprintCallable, Category = "Interchange | Node | Texture")
+	bool SetCustomSRGB(const bool& AttributeValue, bool bAddApplyDelegate = true)
+	{
+		IMPLEMENT_NODE_ATTRIBUTE_SETTER(UInterchangeTextureFactoryNode, SRGB, bool, UTexture)
+	}
+
+	/** Return false if the Attribute was not set previously.*/
+	UFUNCTION(BlueprintCallable, Category = "Interchange | Node | Texture")
+	bool GetCustombUseLegacyGamma(bool& AttributeValue) const
+	{
+		IMPLEMENT_NODE_ATTRIBUTE_GETTER(bUseLegacyGamma, bool);
+	}
+
+	UFUNCTION(BlueprintCallable, Category = "Interchange | Node | Texture")
+	bool SetCustombUseLegacyGamma(const bool& AttributeValue, bool bAddApplyDelegate = true)
+	{
+#if WITH_EDITORONLY_DATA
+		IMPLEMENT_NODE_ATTRIBUTE_SETTER(UInterchangeTextureFactoryNode, bUseLegacyGamma, bool, UTexture)
+#else
+		IMPLEMENT_NODE_ATTRIBUTE_SETTER_NODELEGATE(bUseLegacyGamma, bool)
+#endif
+	}
+
+	/** Return false if the Attribute was not set previously.*/
+	UFUNCTION(BlueprintCallable, Category = "Interchange | Node | Texture")
+	bool GetCustomVirtualTextureStreaming(bool& AttributeValue) const
+	{
+		IMPLEMENT_NODE_ATTRIBUTE_GETTER(VirtualTextureStreaming, bool);
+	}
+
+	UFUNCTION(BlueprintCallable, Category = "Interchange | Node | Texture")
+	bool SetCustomVirtualTextureStreaming(const bool& AttributeValue, bool bAddApplyDelegate = true)
+	{
+		IMPLEMENT_NODE_ATTRIBUTE_SETTER(UInterchangeTextureFactoryNode, VirtualTextureStreaming, bool, UTexture)
+	}
+
+	/** Return false if the Attribute was not set previously.*/
+	UFUNCTION(BlueprintCallable, Category = "Interchange | Node | Texture")
+	bool GetCustomAddressX(uint8& AttributeValue) const
+	{
+		IMPLEMENT_NODE_ATTRIBUTE_GETTER(AddressX, uint8);
+	}
+
+	UFUNCTION(BlueprintCallable, Category = "Interchange | Node | Texture")
+	bool SetCustomAddressX(const uint8& AttributeValue, bool bAddApplyDelegate = true)
+	{
+		IMPLEMENT_NODE_ATTRIBUTE_SETTER(UInterchangeTextureFactoryNode, AddressX, uint8, UTexture2D)
+	}
+
+	/** Return false if the Attribute was not set previously.*/
+	UFUNCTION(BlueprintCallable, Category = "Interchange | Node | Texture")
+	bool GetCustomAddressY(uint8& AttributeValue) const
+	{
+		IMPLEMENT_NODE_ATTRIBUTE_GETTER(AddressY, uint8);
+	}
+
+	UFUNCTION(BlueprintCallable, Category = "Interchange | Node | Texture")
+	bool SetCustomAddressY(const uint8& AttributeValue, bool bAddApplyDelegate = true)
+	{
+		IMPLEMENT_NODE_ATTRIBUTE_SETTER(UInterchangeTextureFactoryNode, AddressY, uint8, UTexture2D)
+	}
+
+	//////////////////////////////////////////////////////////////////////////
+	//Level of Detail
+
+	/** Return false if the Attribute was not set previously.*/
+	UFUNCTION(BlueprintCallable, Category = "Interchange | Node | Texture")
+	bool GetCustombPreserveBorder(bool& AttributeValue) const
+	{
+		IMPLEMENT_NODE_ATTRIBUTE_GETTER(bPreserveBorder, bool);
+	}
+
+	UFUNCTION(BlueprintCallable, Category = "Interchange | Node | Texture")
+	bool SetCustombPreserveBorder(const bool& AttributeValue, bool bAddApplyDelegate = true)
+	{
+#if WITH_EDITORONLY_DATA
+		IMPLEMENT_NODE_ATTRIBUTE_SETTER(UInterchangeTextureFactoryNode, bPreserveBorder, bool, UTexture)
+#else
+		IMPLEMENT_NODE_ATTRIBUTE_SETTER_NODELEGATE(bPreserveBorder, bool)
+#endif
+	}
+
+	/** Return false if the Attribute was not set previously.*/
+	UFUNCTION(BlueprintCallable, Category = "Interchange | Node | Texture")
+	bool GetCustomMipGenSettings(uint8& AttributeValue) const
+	{
+		IMPLEMENT_NODE_ATTRIBUTE_GETTER(MipGenSettings, uint8);
+	}
+
+	UFUNCTION(BlueprintCallable, Category = "Interchange | Node | Texture")
+	bool SetCustomMipGenSettings(const uint8& AttributeValue, bool bAddApplyDelegate = true)
+	{
+#if WITH_EDITORONLY_DATA
+		IMPLEMENT_NODE_ATTRIBUTE_SETTER(UInterchangeTextureFactoryNode, MipGenSettings, uint8, UTexture)
+#else
+		IMPLEMENT_NODE_ATTRIBUTE_SETTER_NODELEGATE(MipGenSettings, uint8)
+#endif
+	}
+
+	/** Return false if the Attribute was not set previously.*/
+	UFUNCTION(BlueprintCallable, Category = "Interchange | Node | Texture")
+	bool GetCustomLODBias(int32& AttributeValue) const
+	{
+		IMPLEMENT_NODE_ATTRIBUTE_GETTER(LODBias, int32);
+	}
+
+	UFUNCTION(BlueprintCallable, Category = "Interchange | Node | Texture")
+	bool SetCustomLODBias(const int32& AttributeValue, bool bAddApplyDelegate = true)
+	{
+		IMPLEMENT_NODE_ATTRIBUTE_SETTER(UInterchangeTextureFactoryNode, LODBias, int32, UTexture)
+	}
+
+	/** Return false if the Attribute was not set previously.*/
+	UFUNCTION(BlueprintCallable, Category = "Interchange | Node | Texture")
+	bool GetCustomLODGroup(uint8& AttributeValue) const
+	{
+		IMPLEMENT_NODE_ATTRIBUTE_GETTER(LODGroup, uint8);
+	}
+
+	UFUNCTION(BlueprintCallable, Category = "Interchange | Node | Texture")
+	bool SetCustomLODGroup(const uint8& AttributeValue, bool bAddApplyDelegate = true)
+	{
+		IMPLEMENT_NODE_ATTRIBUTE_SETTER(UInterchangeTextureFactoryNode, LODGroup, uint8, UTexture)
+	}
+
+	//TODO support per platform data in the FAttributeStorage, so we can set different value per platform at the pipeline stage, We set only the default value for now
+	/** Return false if the Attribute was not set previously.*/
+	UFUNCTION(BlueprintCallable, Category = "Interchange | Node | Texture")
+	bool GetCustomDownscale(float& AttributeValue) const
+	{
+		IMPLEMENT_NODE_ATTRIBUTE_GETTER(Downscale, float);
+	}
+
+	UFUNCTION(BlueprintCallable, Category = "Interchange | Node | Texture")
+	bool SetCustomDownscale(const float& AttributeValue, bool bAddApplyDelegate = true)
+	{
+		IMPLEMENT_NODE_ATTRIBUTE_SETTER(UInterchangeTextureFactoryNode, Downscale, float, UTexture)
+	}
+
+	/** Return false if the Attribute was not set previously.*/
+	UFUNCTION(BlueprintCallable, Category = "Interchange | Node | Texture")
+	bool GetCustomDownscaleOptions(uint8& AttributeValue) const
+	{
+		IMPLEMENT_NODE_ATTRIBUTE_GETTER(DownscaleOptions, uint8);
+	}
+
+	UFUNCTION(BlueprintCallable, Category = "Interchange | Node | Texture")
+	bool SetCustomDownscaleOptions(const uint8& AttributeValue, bool bAddApplyDelegate = true)
+	{
+		IMPLEMENT_NODE_ATTRIBUTE_SETTER(UInterchangeTextureFactoryNode, DownscaleOptions, uint8, UTexture)
+	}
+
+	//////////////////////////////////////////////////////////////////////////
+	//Compositing
+
+	/** Return false if the Attribute was not set previously.*/
+	UFUNCTION(BlueprintCallable, Category = "Interchange | Node | Texture")
+	bool GetCustomCompositeTextureMode(uint8& AttributeValue) const
+	{
+		IMPLEMENT_NODE_ATTRIBUTE_GETTER(CompositeTextureMode, uint8);
+	}
+
+	UFUNCTION(BlueprintCallable, Category = "Interchange | Node | Texture")
+	bool SetCustomCompositeTextureMode(const uint8& AttributeValue, bool bAddApplyDelegate = true)
+	{
+#if WITH_EDITORONLY_DATA
+		IMPLEMENT_NODE_ATTRIBUTE_SETTER(UInterchangeTextureFactoryNode, CompositeTextureMode, uint8, UTexture)
+#else
+		IMPLEMENT_NODE_ATTRIBUTE_SETTER_NODELEGATE(CompositeTextureMode, uint8)
+#endif
+	}
+
+	/** Return false if the Attribute was not set previously.*/
+	UFUNCTION(BlueprintCallable, Category = "Interchange | Node | Texture")
+	bool GetCustomCompositePower(float& AttributeValue) const
+	{
+		IMPLEMENT_NODE_ATTRIBUTE_GETTER(CompositePower, float);
+	}
+
+	UFUNCTION(BlueprintCallable, Category = "Interchange | Node | Texture")
+	bool SetCustomCompositePower(const float& AttributeValue, bool bAddApplyDelegate = true)
+	{
+#if WITH_EDITORONLY_DATA
+		IMPLEMENT_NODE_ATTRIBUTE_SETTER(UInterchangeTextureFactoryNode, CompositePower, float, UTexture)
+#else
+		IMPLEMENT_NODE_ATTRIBUTE_SETTER_NODELEGATE(CompositePower, float)
+#endif
+	}
+
+private:
+
+	void FillAssetClassFromAttribute()
+	{
+#if WITH_ENGINE
+		FString OperationName = GetTypeName() + TEXT(".GetAssetClassName");
+		FString ClassName;
+		InterchangePrivateNodeBase::GetCustomAttribute<FString>(*Attributes, ClassNameAttributeKey, OperationName, ClassName);
+		if (ClassName.Equals(UTexture2D::StaticClass()->GetName()))
+		{
+			AssetClass = UTexture2D::StaticClass();
+			bIsTextureNodeClassInitialized = true;
+		}
+		else if (ClassName.Equals(UTexture::StaticClass()->GetName()))
+		{
+			AssetClass = UTexture::StaticClass();
+			bIsTextureNodeClassInitialized = true;
+		}
+#endif
+	}
+
+	bool SetTextureNodeClassFromClassAttribute()
+	{
+		if (!bIsTextureNodeClassInitialized)
+		{
+			FillAssetClassFromAttribute();
+		}
+		return bIsTextureNodeClassInitialized;
+	}
+
+	const UE::Interchange::FAttributeKey ClassNameAttributeKey = UE::Interchange::FBaseNodeStaticData::ClassTypeAttributeKey();
+	const UE::Interchange::FAttributeKey AssetNameKey = UE::Interchange::FBaseNodeStaticData::AssetNameKey();
+
+	//Translated texture node unique id
+	const UE::Interchange::FAttributeKey Macro_CustomTranslatedTextureNodeUidKey = UE::Interchange::FAttributeKey(TEXT("TranslatedTextureNodeUid"));
+	
+	//Texture Adjustments
+	const UE::Interchange::FAttributeKey Macro_CustomAdjustBrightnessKey = UE::Interchange::FAttributeKey(TEXT("AdjustBrightness"));
+	const UE::Interchange::FAttributeKey Macro_CustomAdjustBrightnessCurveKey = UE::Interchange::FAttributeKey(TEXT("AdjustBrightnessCurve"));
+	const UE::Interchange::FAttributeKey Macro_CustomAdjustVibranceKey = UE::Interchange::FAttributeKey(TEXT("AdjustVibrance"));
+	const UE::Interchange::FAttributeKey Macro_CustomAdjustSaturationKey = UE::Interchange::FAttributeKey(TEXT("AdjustSaturation"));
+	const UE::Interchange::FAttributeKey Macro_CustomAdjustRGBCurveKey = UE::Interchange::FAttributeKey(TEXT("AdjustRGBCurve"));
+	const UE::Interchange::FAttributeKey Macro_CustomAdjustHueKey = UE::Interchange::FAttributeKey(TEXT("AdjustHue"));
+	const UE::Interchange::FAttributeKey Macro_CustomAdjustMinAlphaKey = UE::Interchange::FAttributeKey(TEXT("AdjustMinAlpha"));
+	const UE::Interchange::FAttributeKey Macro_CustomAdjustMaxAlphaKey = UE::Interchange::FAttributeKey(TEXT("AdjustMaxAlpha"));
+	const UE::Interchange::FAttributeKey Macro_CustombChromaKeyTextureKey = UE::Interchange::FAttributeKey(TEXT("bChromaKeyTexture"));
+	const UE::Interchange::FAttributeKey Macro_CustomChromaKeyThresholdKey = UE::Interchange::FAttributeKey(TEXT("ChromaKeyThreshold"));
+	const UE::Interchange::FAttributeKey Macro_CustomChromaKeyColorKey = UE::Interchange::FAttributeKey(TEXT("ChromaKeyColor"));
+
+	//Texture Compression
+	const UE::Interchange::FAttributeKey Macro_CustomCompressionNoAlphaKey = UE::Interchange::FAttributeKey(TEXT("CompressionNoAlpha"));
+	const UE::Interchange::FAttributeKey Macro_CustomDeferCompressionKey = UE::Interchange::FAttributeKey(TEXT("DeferCompression"));
+	const UE::Interchange::FAttributeKey Macro_CustomLossyCompressionAmountKey = UE::Interchange::FAttributeKey(TEXT("LossyCompressionAmount"));
+	const UE::Interchange::FAttributeKey Macro_CustomMaxTextureSizeKey = UE::Interchange::FAttributeKey(TEXT("MaxTextureSize"));
+	const UE::Interchange::FAttributeKey Macro_CustomCompressionQualityKey = UE::Interchange::FAttributeKey(TEXT("CompressionQuality"));
+	const UE::Interchange::FAttributeKey Macro_CustomCompressionSettingsKey = UE::Interchange::FAttributeKey(TEXT("CompressionSettings"));
+
+	//Texture general
+	const UE::Interchange::FAttributeKey Macro_CustombDitherMipMapAlphaKey = UE::Interchange::FAttributeKey(TEXT("bDitherMipMapAlpha"));
+	const UE::Interchange::FAttributeKey Macro_CustomAlphaCoverageThresholdsKey = UE::Interchange::FAttributeKey(TEXT("AlphaCoverageThresholds"));
+	const UE::Interchange::FAttributeKey Macro_CustombFlipGreenChannelKey = UE::Interchange::FAttributeKey(TEXT("bFlipGreenChannel"));
+	const UE::Interchange::FAttributeKey Macro_CustombForcePVRTC4Key = UE::Interchange::FAttributeKey(TEXT("bForcePVRTC4"));
+	const UE::Interchange::FAttributeKey Macro_CustomPowerOfTwoModeKey = UE::Interchange::FAttributeKey(TEXT("PowerOfTwoMode"));
+	const UE::Interchange::FAttributeKey Macro_CustomPaddingColorKey = UE::Interchange::FAttributeKey(TEXT("PaddingColor"));
+	const UE::Interchange::FAttributeKey Macro_CustomFilterKey = UE::Interchange::FAttributeKey(TEXT("Filter"));
+	const UE::Interchange::FAttributeKey Macro_CustomMipLoadOptionsKey = UE::Interchange::FAttributeKey(TEXT("MipLoadOptions"));
+	const UE::Interchange::FAttributeKey Macro_CustomSRGBKey = UE::Interchange::FAttributeKey(TEXT("SRGB"));
+	const UE::Interchange::FAttributeKey Macro_CustombUseLegacyGammaKey = UE::Interchange::FAttributeKey(TEXT("bUseLegacyGamma"));
+	const UE::Interchange::FAttributeKey Macro_CustomVirtualTextureStreamingKey = UE::Interchange::FAttributeKey(TEXT("VirtualTextureStreaming"));
+	const UE::Interchange::FAttributeKey Macro_CustomAddressXKey = UE::Interchange::FAttributeKey(TEXT("AddressX"));
+	const UE::Interchange::FAttributeKey Macro_CustomAddressYKey = UE::Interchange::FAttributeKey(TEXT("AddressY"));
+
+	//Level of Detail
+	const UE::Interchange::FAttributeKey Macro_CustombPreserveBorderKey = UE::Interchange::FAttributeKey(TEXT("bPreserveBorder"));
+	const UE::Interchange::FAttributeKey Macro_CustomMipGenSettingsKey = UE::Interchange::FAttributeKey(TEXT("MipGenSettings"));
+	const UE::Interchange::FAttributeKey Macro_CustomLODBiasKey = UE::Interchange::FAttributeKey(TEXT("LODBias"));
+	const UE::Interchange::FAttributeKey Macro_CustomLODGroupKey = UE::Interchange::FAttributeKey(TEXT("LODGroup"));
+	//TODO support per platform data in the FAttributeStorage, so we can set different value per platform at the pipeline stage, We set only the default value for now
+	const UE::Interchange::FAttributeKey Macro_CustomDownscaleKey = UE::Interchange::FAttributeKey(TEXT("Downscale"));
+	const UE::Interchange::FAttributeKey Macro_CustomDownscaleOptionsKey = UE::Interchange::FAttributeKey(TEXT("DownscaleOptions"));
+
+	//Compositing
+	const UE::Interchange::FAttributeKey Macro_CustomCompositeTextureModeKey = UE::Interchange::FAttributeKey(TEXT("CompositeTextureMode"));
+	const UE::Interchange::FAttributeKey Macro_CustomCompositePowerKey = UE::Interchange::FAttributeKey(TEXT("CompositePower"));
+
+#if WITH_EDITORONLY_DATA
+	IMPLEMENT_NODE_ATTRIBUTE_APPLY_UOBJECT(AdjustBrightness, float, UTexture, );
+	IMPLEMENT_NODE_ATTRIBUTE_APPLY_UOBJECT(AdjustBrightnessCurve, float, UTexture, );
+	IMPLEMENT_NODE_ATTRIBUTE_APPLY_UOBJECT(AdjustVibrance, float, UTexture, );
+	IMPLEMENT_NODE_ATTRIBUTE_APPLY_UOBJECT(AdjustSaturation, float, UTexture, );
+	IMPLEMENT_NODE_ATTRIBUTE_APPLY_UOBJECT(AdjustRGBCurve, float, UTexture, );
+	IMPLEMENT_NODE_ATTRIBUTE_APPLY_UOBJECT(AdjustHue, float, UTexture, );
+	IMPLEMENT_NODE_ATTRIBUTE_APPLY_UOBJECT(AdjustMinAlpha, float, UTexture, );
+	IMPLEMENT_NODE_ATTRIBUTE_APPLY_UOBJECT(AdjustMaxAlpha, float, UTexture, );
+	IMPLEMENT_NODE_ATTRIBUTE_APPLY_UOBJECT(bChromaKeyTexture, bool, UTexture, );
+	IMPLEMENT_NODE_ATTRIBUTE_APPLY_UOBJECT(ChromaKeyThreshold, float, UTexture, );
+	IMPLEMENT_NODE_ATTRIBUTE_APPLY_UOBJECT(ChromaKeyColor, FColor, UTexture, );
+
+	//Texture Compression
+	IMPLEMENT_NODE_ATTRIBUTE_APPLY_UOBJECT(CompressionNoAlpha, bool, UTexture, );
+	IMPLEMENT_NODE_ATTRIBUTE_APPLY_UOBJECT(DeferCompression, bool, UTexture, );
+	IMPLEMENT_NODE_ATTRIBUTE_APPLY_UOBJECT(LossyCompressionAmount, uint8, UTexture, TEnumAsByte<ETextureLossyCompressionAmount>);
+	IMPLEMENT_NODE_ATTRIBUTE_APPLY_UOBJECT(MaxTextureSize, int32, UTexture, );
+	IMPLEMENT_NODE_ATTRIBUTE_APPLY_UOBJECT(CompressionQuality, uint8, UTexture, TEnumAsByte<enum ETextureCompressionQuality>);
+#endif
+	IMPLEMENT_NODE_ATTRIBUTE_APPLY_UOBJECT(CompressionSettings, uint8, UTexture, TEnumAsByte<enum TextureCompressionSettings>);
+#if WITH_EDITORONLY_DATA
+	//Texture general
+	IMPLEMENT_NODE_ATTRIBUTE_APPLY_UOBJECT(bDitherMipMapAlpha, bool, UTexture, );
+	IMPLEMENT_NODE_ATTRIBUTE_APPLY_UOBJECT(AlphaCoverageThresholds, FVector4, UTexture, );
+	IMPLEMENT_NODE_ATTRIBUTE_APPLY_UOBJECT(bFlipGreenChannel, bool, UTexture, );
+	IMPLEMENT_NODE_ATTRIBUTE_APPLY_UOBJECT(PowerOfTwoMode, uint8, UTexture, TEnumAsByte<enum ETexturePowerOfTwoSetting::Type>);
+	IMPLEMENT_NODE_ATTRIBUTE_APPLY_UOBJECT(PaddingColor, FColor, UTexture, );
+#endif //#if WITH_EDITORONLY_DATA
+	IMPLEMENT_NODE_ATTRIBUTE_APPLY_UOBJECT(Filter, uint8, UTexture, TEnumAsByte<enum TextureFilter>);
+	IMPLEMENT_NODE_ATTRIBUTE_APPLY_UOBJECT(MipLoadOptions, uint8, UTexture, ETextureMipLoadOptions);
+	IMPLEMENT_NODE_ATTRIBUTE_APPLY_UOBJECT(SRGB, bool, UTexture, );
+#if WITH_EDITORONLY_DATA
+	IMPLEMENT_NODE_ATTRIBUTE_APPLY_UOBJECT(bUseLegacyGamma, bool, UTexture, );
+#endif //#if WITH_EDITORONLY_DATA
+	IMPLEMENT_NODE_ATTRIBUTE_APPLY_UOBJECT(VirtualTextureStreaming, bool, UTexture, );
+	IMPLEMENT_NODE_ATTRIBUTE_APPLY_UOBJECT(AddressX, uint8, UTexture2D, TEnumAsByte<enum TextureAddress>);
+	IMPLEMENT_NODE_ATTRIBUTE_APPLY_UOBJECT(AddressY, uint8, UTexture2D, TEnumAsByte<enum TextureAddress>);
+#if WITH_EDITORONLY_DATA
+	//Level of Detail
+	IMPLEMENT_NODE_ATTRIBUTE_APPLY_UOBJECT(bPreserveBorder, bool, UTexture, );
+	IMPLEMENT_NODE_ATTRIBUTE_APPLY_UOBJECT(MipGenSettings, uint8, UTexture, TEnumAsByte<enum TextureMipGenSettings>);
+#endif //#if WITH_EDITORONLY_DATA
+	IMPLEMENT_NODE_ATTRIBUTE_APPLY_UOBJECT(LODBias, int32, UTexture, );
+	IMPLEMENT_NODE_ATTRIBUTE_APPLY_UOBJECT(LODGroup, uint8, UTexture, TEnumAsByte<enum TextureGroup>);
+
+	//TODO support per platform data in the FAttributeStorage, so we can set different value per platform at the pipeline stage, We set only the default value for now
+	//IMPLEMENT_NODE_ATTRIBUTE_APPLY_UOBJECT(Downscale, float, UTexture, );
+#if WITH_ENGINE
+	bool ApplyCustomDownscaleToAsset(UObject* Asset) const
+	{
+		if (!Asset)
+		{
+			return false;
+		}
+		UTexture* TypedObject = Cast<UTexture>(Asset);
+		if (!TypedObject)
+		{
+			return false;
+		}
+		float ValueData;
+		if (GetCustomDownscale(ValueData))
+		{
+			TypedObject->Downscale.Default = ValueData;
+			return true;
+		}
+		return false;
+	}
+
+	bool FillCustomDownscaleFromAsset(UObject* Asset)
+	{
+		if (!Asset)
+		{
+			return false;
+		}
+		UTexture* TypedObject = Cast<UTexture>(Asset);
+		if (!TypedObject)
+		{
+			return false;
+		}
+		if (SetCustomDownscale(TypedObject->Downscale.Default, false))
+		{
+			return true;
+		}
+		return false;
+	}
+#endif //WITH_ENGINE
+
+	IMPLEMENT_NODE_ATTRIBUTE_APPLY_UOBJECT(DownscaleOptions, uint8, UTexture, ETextureDownscaleOptions);
+#if WITH_EDITORONLY_DATA
+	//Compositing
+	IMPLEMENT_NODE_ATTRIBUTE_APPLY_UOBJECT(CompositeTextureMode, uint8, UTexture, TEnumAsByte<enum ECompositeTextureMode>);
+	IMPLEMENT_NODE_ATTRIBUTE_APPLY_UOBJECT(CompositePower, float, UTexture, );
+#endif //#if WITH_EDITORONLY_DATA
+protected:
+#if WITH_ENGINE
+	TSubclassOf<UTexture> AssetClass = nullptr;
+#endif
+
+	UE::Interchange::TMapAttributeHelper<int32, FString> SourceBlocks;
+
+	bool bIsTextureNodeClassInitialized = false;
+};
diff --git a/Engine/Plugins/Runtime/Interchange/Source/Nodes/Public/InterchangeTextureNode.h b/Engine/Plugins/Runtime/Interchange/Source/Nodes/Public/InterchangeTextureNode.h
--- a/Engine/Plugins/Runtime/Interchange/Source/Nodes/Public/InterchangeTextureNode.h
+++ b/Engine/Plugins/Runtime/Interchange/Source/Nodes/Public/InterchangeTextureNode.h
@@ -5,12 +5,6 @@
 #include "CoreMinimal.h"
 #include "Nodes/InterchangeBaseNode.h"
 
-#if WITH_ENGINE
-#include "Engine/Texture.h"
-#include "Engine/Texture2D.h"
-#endif
-
-
 #include "InterchangeTextureNode.generated.h"
 
 //Interchange namespace
@@ -33,7 +27,6 @@ namespace UE
 				return StringKey;
 			}
 		};
-
 	}//ns Interchange
 }//ns UE
 
@@ -46,9 +39,6 @@ public:
 	UInterchangeTextureNode()
 	:UInterchangeBaseNode()
 	{
-#if WITH_ENGINE
-		AssetClass = nullptr;
-#endif
 	}
 
 	virtual void PostInitProperties()
@@ -57,35 +47,6 @@ public:
 		SourceBlocks.Initialize(Attributes.ToSharedRef(), UE::Interchange::FTextureNodeStaticData::GetBaseSourceBlocksKey());
 	}
 
-	/**
-	 * Initialize node data
-	 * @param: UniqueID - The uniqueId for this node
-	 * @param DisplayLabel - The name of the node
-	 * @param InAssetClass - The class the texture factory will create for this node.
-	 */
-	UFUNCTION(BlueprintCallable, Category = "Interchange | Node | Texture")
-	void InitializeTextureNode(const FString& UniqueID, const FString& DisplayLabel, const FString& InAssetClass)
-	{
-		bIsTextureNodeClassInitialized = false;
-		InitializeNode(UniqueID, DisplayLabel);
-		FString OperationName = GetTypeName() + TEXT(".SetAssetClassName");
-		InterchangePrivateNodeBase::SetCustomAttribute<FString>(*Attributes, ClassNameAttributeKey, OperationName, InAssetClass);
-		FillAssetClassFromAttribute();
-	}
-
-	virtual void Serialize(FArchive& Ar) override
-	{
-		Super::Serialize(Ar);
-#if WITH_ENGINE
-		if(Ar.IsLoading())
-		{
-			//Make sure the class is properly set when we compile with engine, this will set the
-			//bIsTextureNodeClassInitialized to true.
-			SetTextureNodeClassFromClassAttribute();
-		}
-#endif //#if WITH_ENGINE
-	}
-
 	/**
 	 * Return the node type name of the class, we use this when reporting error
 	 */
@@ -105,18 +66,6 @@ public:
 		return Super::GetKeyDisplayName(NodeAttributeKey);
 	}
 
-	/** Get the class this node want to create */
-	UFUNCTION(BlueprintCallable, Category = "Interchange | Node | Texture")
-	virtual class UClass* GetAssetClass() const override
-	{
-		ensure(bIsTextureNodeClassInitialized);
-#if WITH_ENGINE
-		return AssetClass.Get() != nullptr ? AssetClass.Get() : UTexture::StaticClass();
-#else
-		return nullptr;
-#endif
-	}
-
 	virtual FGuid GetHash() const override
 	{
 		return Attributes->GetStorageHash();
@@ -154,17 +103,16 @@ public:
 		}
 	}
 
-public:
 	//////////////////////////////////////////////////////////////////////////
 	// UDIMs begin here
 	// UDIM base texture use a different model for the source data
-	
+
 	/**
 	 * Get the source blocks for the texture
 	 * If the map is empty then the texture will be simply be imported as normal texture using the payload key
 	 */
 	UFUNCTION(BlueprintCallable, Category = "Interchange | Node | Texture | UDIMs")
-	TMap<int32,FString> GetSourceBlocks() const
+	TMap<int32, FString> GetSourceBlocks() const
 	{
 		return SourceBlocks.ToMap();
 	}
@@ -177,7 +125,7 @@ public:
 	 * The first block inserted in the map be used to determine the accepted texture format and pixel format
 	 */
 	UFUNCTION(BlueprintCallable, Category = "Interchange | Node | Texture | UDIMs")
-	void SetSourceBlocks(const TMap<int32,FString>& InSourceBlocks)
+	void SetSourceBlocks(const TMap<int32, FString>& InSourceBlocks)
 	{
 		SourceBlocks = InSourceBlocks;
 	}
@@ -194,796 +142,8 @@ public:
 
 	// UDIMs ends here
 	//////////////////////////////////////////////////////////////////////////
-	
-public:
-	/** Return false if the Attribute was not set previously.*/
-	UFUNCTION(BlueprintCallable, Category = "Interchange | Node | Texture")
-	bool GetCustomAdjustBrightness(float& AttributeValue) const
-	{
-		IMPLEMENT_NODE_ATTRIBUTE_GETTER(AdjustBrightness, float);
-	}
-
-	UFUNCTION(BlueprintCallable, Category = "Interchange | Node | Texture")
-	bool SetCustomAdjustBrightness(const float& AttributeValue, bool bAddApplyDelegate = true)
-	{
-#if WITH_EDITORONLY_DATA
-		IMPLEMENT_NODE_ATTRIBUTE_SETTER(UInterchangeTextureNode, AdjustBrightness, float, UTexture)
-#else
-		IMPLEMENT_NODE_ATTRIBUTE_SETTER_NODELEGATE(AdjustBrightness, float)
-#endif
-	}
-
-	/** Return false if the Attribute was not set previously.*/
-	UFUNCTION(BlueprintCallable, Category = "Interchange | Node | Texture")
-	bool GetCustomAdjustBrightnessCurve(float& AttributeValue) const
-	{
-		IMPLEMENT_NODE_ATTRIBUTE_GETTER(AdjustBrightnessCurve, float);
-	}
-
-	UFUNCTION(BlueprintCallable, Category = "Interchange | Node | Texture")
-	bool SetCustomAdjustBrightnessCurve(const float& AttributeValue, bool bAddApplyDelegate = true)
-	{
-#if WITH_EDITORONLY_DATA
-		IMPLEMENT_NODE_ATTRIBUTE_SETTER(UInterchangeTextureNode, AdjustBrightnessCurve, float, UTexture)
-#else
-		IMPLEMENT_NODE_ATTRIBUTE_SETTER_NODELEGATE(AdjustBrightnessCurve, float)
-#endif
-	}
-
-	/** Return false if the Attribute was not set previously.*/
-	UFUNCTION(BlueprintCallable, Category = "Interchange | Node | Texture")
-	bool GetCustomAdjustVibrance(float& AttributeValue) const
-	{
-		IMPLEMENT_NODE_ATTRIBUTE_GETTER(AdjustVibrance, float);
-	}
-
-	UFUNCTION(BlueprintCallable, Category = "Interchange | Node | Texture")
-	bool SetCustomAdjustVibrance(const float& AttributeValue, bool bAddApplyDelegate = true)
-	{
-#if WITH_EDITORONLY_DATA
-		IMPLEMENT_NODE_ATTRIBUTE_SETTER(UInterchangeTextureNode, AdjustVibrance, float, UTexture)
-#else
-		IMPLEMENT_NODE_ATTRIBUTE_SETTER_NODELEGATE(AdjustVibrance, float)
-#endif
-	}
-
-	/** Return false if the Attribute was not set previously.*/
-	UFUNCTION(BlueprintCallable, Category = "Interchange | Node | Texture")
-	bool GetCustomAdjustSaturation(float& AttributeValue) const
-	{
-		IMPLEMENT_NODE_ATTRIBUTE_GETTER(AdjustSaturation, float);
-	}
-
-	UFUNCTION(BlueprintCallable, Category = "Interchange | Node | Texture")
-	bool SetCustomAdjustSaturation(const float& AttributeValue, bool bAddApplyDelegate = true)
-	{
-#if WITH_EDITORONLY_DATA
-		IMPLEMENT_NODE_ATTRIBUTE_SETTER(UInterchangeTextureNode, AdjustSaturation, float, UTexture)
-#else
-		IMPLEMENT_NODE_ATTRIBUTE_SETTER_NODELEGATE(AdjustSaturation, float)
-#endif
-	}
-
-	/** Return false if the Attribute was not set previously.*/
-	UFUNCTION(BlueprintCallable, Category = "Interchange | Node | Texture")
-	bool GetCustomAdjustRGBCurve(float& AttributeValue) const
-	{
-		IMPLEMENT_NODE_ATTRIBUTE_GETTER(AdjustRGBCurve, float);
-	}
-
-	UFUNCTION(BlueprintCallable, Category = "Interchange | Node | Texture")
-	bool SetCustomAdjustRGBCurve(const float& AttributeValue, bool bAddApplyDelegate = true)
-	{
-#if WITH_EDITORONLY_DATA
-		IMPLEMENT_NODE_ATTRIBUTE_SETTER(UInterchangeTextureNode, AdjustRGBCurve, float, UTexture)
-#else
-		IMPLEMENT_NODE_ATTRIBUTE_SETTER_NODELEGATE(AdjustRGBCurve, float)
-#endif
-	}
-
-	/** Return false if the Attribute was not set previously.*/
-	UFUNCTION(BlueprintCallable, Category = "Interchange | Node | Texture")
-	bool GetCustomAdjustHue(float& AttributeValue) const
-	{
-		IMPLEMENT_NODE_ATTRIBUTE_GETTER(AdjustHue, float);
-	}
-
-	UFUNCTION(BlueprintCallable, Category = "Interchange | Node | Texture")
-	bool SetCustomAdjustHue(const float& AttributeValue, bool bAddApplyDelegate = true)
-	{
-#if WITH_EDITORONLY_DATA
-		IMPLEMENT_NODE_ATTRIBUTE_SETTER(UInterchangeTextureNode, AdjustHue, float, UTexture)
-#else
-		IMPLEMENT_NODE_ATTRIBUTE_SETTER_NODELEGATE(AdjustHue, float)
-#endif
-	}
-
-	/** Return false if the Attribute was not set previously.*/
-	UFUNCTION(BlueprintCallable, Category = "Interchange | Node | Texture")
-	bool GetCustomAdjustMinAlpha(float& AttributeValue) const
-	{
-		IMPLEMENT_NODE_ATTRIBUTE_GETTER(AdjustMinAlpha, float);
-	}
-
-	UFUNCTION(BlueprintCallable, Category = "Interchange | Node | Texture")
-	bool SetCustomAdjustMinAlpha(const float& AttributeValue, bool bAddApplyDelegate = true)
-	{
-#if WITH_EDITORONLY_DATA
-		IMPLEMENT_NODE_ATTRIBUTE_SETTER(UInterchangeTextureNode, AdjustMinAlpha, float, UTexture)
-#else
-		IMPLEMENT_NODE_ATTRIBUTE_SETTER_NODELEGATE(AdjustMinAlpha, float)
-#endif
-	}
-
-	/** Return false if the Attribute was not set previously.*/
-	UFUNCTION(BlueprintCallable, Category = "Interchange | Node | Texture")
-	bool GetCustomAdjustMaxAlpha(float& AttributeValue) const
-	{
-		IMPLEMENT_NODE_ATTRIBUTE_GETTER(AdjustMaxAlpha, float);
-	}
-
-	UFUNCTION(BlueprintCallable, Category = "Interchange | Node | Texture")
-	bool SetCustomAdjustMaxAlpha(const float& AttributeValue, bool bAddApplyDelegate = true)
-	{
-#if WITH_EDITORONLY_DATA
-		IMPLEMENT_NODE_ATTRIBUTE_SETTER(UInterchangeTextureNode, AdjustMaxAlpha, float, UTexture)
-#else
-		IMPLEMENT_NODE_ATTRIBUTE_SETTER_NODELEGATE(AdjustMaxAlpha, float)
-#endif
-	}
-
-	/** Return false if the Attribute was not set previously.*/
-	UFUNCTION(BlueprintCallable, Category = "Interchange | Node | Texture")
-	bool GetCustombChromaKeyTexture(bool& AttributeValue) const
-	{
-		IMPLEMENT_NODE_ATTRIBUTE_GETTER(bChromaKeyTexture, bool);
-	}
-
-	UFUNCTION(BlueprintCallable, Category = "Interchange | Node | Texture")
-	bool SetCustombChromaKeyTexture(const bool& AttributeValue, bool bAddApplyDelegate = true)
-	{
-#if WITH_EDITORONLY_DATA
-		IMPLEMENT_NODE_ATTRIBUTE_SETTER(UInterchangeTextureNode, bChromaKeyTexture, bool, UTexture)
-#else
-		IMPLEMENT_NODE_ATTRIBUTE_SETTER_NODELEGATE(bChromaKeyTexture, bool)
-#endif
-	}
-
-	/** Return false if the Attribute was not set previously.*/
-	UFUNCTION(BlueprintCallable, Category = "Interchange | Node | Texture")
-	bool GetCustomChromaKeyThreshold(float& AttributeValue) const
-	{
-		IMPLEMENT_NODE_ATTRIBUTE_GETTER(ChromaKeyThreshold, float);
-	}
-
-	UFUNCTION(BlueprintCallable, Category = "Interchange | Node | Texture")
-	bool SetCustomChromaKeyThreshold(const float& AttributeValue, bool bAddApplyDelegate = true)
-	{
-#if WITH_EDITORONLY_DATA
-		IMPLEMENT_NODE_ATTRIBUTE_SETTER(UInterchangeTextureNode, ChromaKeyThreshold, float, UTexture)
-#else
-		IMPLEMENT_NODE_ATTRIBUTE_SETTER_NODELEGATE(ChromaKeyThreshold, float)
-#endif
-	}
-
-	/** Return false if the Attribute was not set previously.*/
-	UFUNCTION(BlueprintCallable, Category = "Interchange | Node | Texture")
-	bool GetCustomChromaKeyColor(FColor& AttributeValue) const
-	{
-		IMPLEMENT_NODE_ATTRIBUTE_GETTER(ChromaKeyColor, FColor);
-	}
-
-	UFUNCTION(BlueprintCallable, Category = "Interchange | Node | Texture")
-	bool SetCustomChromaKeyColor(const FColor& AttributeValue, bool bAddApplyDelegate = true)
-	{
-#if WITH_EDITORONLY_DATA
-		IMPLEMENT_NODE_ATTRIBUTE_SETTER(UInterchangeTextureNode, ChromaKeyColor, FColor, UTexture)
-#else
-		IMPLEMENT_NODE_ATTRIBUTE_SETTER_NODELEGATE(ChromaKeyColor, FColor)
-#endif
-	}
-
-	//////////////////////////////////////////////////////////////////////////
-	//Texture Compression Begin
-
-	/** Return false if the Attribute was not set previously.*/
-	UFUNCTION(BlueprintCallable, Category = "Interchange | Node | Texture")
-	bool GetCustomCompressionNoAlpha(bool& AttributeValue) const
-	{
-		IMPLEMENT_NODE_ATTRIBUTE_GETTER(CompressionNoAlpha, bool);
-	}
-
-	UFUNCTION(BlueprintCallable, Category = "Interchange | Node | Texture")
-	bool SetCustomCompressionNoAlpha(const bool& AttributeValue, bool bAddApplyDelegate = true)
-	{
-#if WITH_EDITORONLY_DATA
-		IMPLEMENT_NODE_ATTRIBUTE_SETTER(UInterchangeTextureNode, CompressionNoAlpha, bool, UTexture)
-#else
-		IMPLEMENT_NODE_ATTRIBUTE_SETTER_NODELEGATE(CompressionNoAlpha, bool)
-#endif
-	}
-
-	/** Return false if the Attribute was not set previously.*/
-	UFUNCTION(BlueprintCallable, Category = "Interchange | Node | Texture")
-	bool GetCustomDeferCompression(bool& AttributeValue) const
-	{
-		IMPLEMENT_NODE_ATTRIBUTE_GETTER(DeferCompression, bool);
-	}
-
-	UFUNCTION(BlueprintCallable, Category = "Interchange | Node | Texture")
-	bool SetCustomDeferCompression(const bool& AttributeValue, bool bAddApplyDelegate = true)
-	{
-#if WITH_EDITORONLY_DATA
-		IMPLEMENT_NODE_ATTRIBUTE_SETTER(UInterchangeTextureNode, DeferCompression, bool, UTexture)
-#else
-		IMPLEMENT_NODE_ATTRIBUTE_SETTER_NODELEGATE(DeferCompression, bool)
-#endif
-	}
-
-	/** Return false if the Attribute was not set previously.*/
-	UFUNCTION(BlueprintCallable, Category = "Interchange | Node | Texture")
-	bool GetCustomLossyCompressionAmount(uint8& AttributeValue) const
-	{
-		IMPLEMENT_NODE_ATTRIBUTE_GETTER(LossyCompressionAmount, uint8);
-	}
-
-	UFUNCTION(BlueprintCallable, Category = "Interchange | Node | Texture")
-	bool SetCustomLossyCompressionAmount(const uint8& AttributeValue, bool bAddApplyDelegate = true)
-	{
-#if WITH_EDITORONLY_DATA
-		IMPLEMENT_NODE_ATTRIBUTE_SETTER(UInterchangeTextureNode, LossyCompressionAmount, uint8, UTexture)
-#else
-		IMPLEMENT_NODE_ATTRIBUTE_SETTER_NODELEGATE(LossyCompressionAmount, uint8)
-#endif
-	}
-
-	/** Return false if the Attribute was not set previously.*/
-	UFUNCTION(BlueprintCallable, Category = "Interchange | Node | Texture")
-	bool GetCustomMaxTextureSize(int32& AttributeValue) const
-	{
-		IMPLEMENT_NODE_ATTRIBUTE_GETTER(MaxTextureSize, int32);
-	}
-
-	UFUNCTION(BlueprintCallable, Category = "Interchange | Node | Texture")
-	bool SetCustomMaxTextureSize(const int32& AttributeValue, bool bAddApplyDelegate = true)
-	{
-#if WITH_EDITORONLY_DATA
-		IMPLEMENT_NODE_ATTRIBUTE_SETTER(UInterchangeTextureNode, MaxTextureSize, int32, UTexture)
-#else
-		IMPLEMENT_NODE_ATTRIBUTE_SETTER_NODELEGATE(MaxTextureSize, int32)
-#endif
-	}
-
-	/** Return false if the Attribute was not set previously.*/
-	UFUNCTION(BlueprintCallable, Category = "Interchange | Node | Texture")
-	bool GetCustomCompressionQuality(uint8& AttributeValue) const
-	{
-		IMPLEMENT_NODE_ATTRIBUTE_GETTER(CompressionQuality, uint8);
-	}
-
-	UFUNCTION(BlueprintCallable, Category = "Interchange | Node | Texture")
-	bool SetCustomCompressionQuality(const uint8& AttributeValue, bool bAddApplyDelegate = true)
-	{
-#if WITH_EDITORONLY_DATA
-		IMPLEMENT_NODE_ATTRIBUTE_SETTER(UInterchangeTextureNode, CompressionQuality, uint8, UTexture)
-#else
-		IMPLEMENT_NODE_ATTRIBUTE_SETTER_NODELEGATE(CompressionQuality, uint8)
-#endif
-	}
-
-	/** Return false if the Attribute was not set previously.*/
-	UFUNCTION(BlueprintCallable, Category = "Interchange | Node | Texture")
-	bool GetCustomCompressionSettings(uint8& AttributeValue) const
-	{
-		IMPLEMENT_NODE_ATTRIBUTE_GETTER(CompressionSettings, uint8);
-	}
-
-	UFUNCTION(BlueprintCallable, Category = "Interchange | Node | Texture")
-	bool SetCustomCompressionSettings(const uint8& AttributeValue, bool bAddApplyDelegate = true)
-	{
-		IMPLEMENT_NODE_ATTRIBUTE_SETTER(UInterchangeTextureNode, CompressionSettings, uint8, UTexture)
-	}
-
-	//////////////////////////////////////////////////////////////////////////
-	//Texture general
-
-	/** Return false if the Attribute was not set previously.*/
-	UFUNCTION(BlueprintCallable, Category = "Interchange | Node | Texture")
-	bool GetCustombDitherMipMapAlpha(bool& AttributeValue) const
-	{
-		IMPLEMENT_NODE_ATTRIBUTE_GETTER(bDitherMipMapAlpha, bool);
-	}
-
-	UFUNCTION(BlueprintCallable, Category = "Interchange | Node | Texture")
-	bool SetCustombDitherMipMapAlpha(const bool& AttributeValue, bool bAddApplyDelegate = true)
-	{
-#if WITH_EDITORONLY_DATA
-		IMPLEMENT_NODE_ATTRIBUTE_SETTER(UInterchangeTextureNode, bDitherMipMapAlpha, bool, UTexture)
-#else
-		IMPLEMENT_NODE_ATTRIBUTE_SETTER_NODELEGATE(bDitherMipMapAlpha, bool)
-#endif
-	}
-
-	/** Return false if the Attribute was not set previously.*/
-	UFUNCTION(BlueprintCallable, Category = "Interchange | Node | Texture")
-	bool GetCustomAlphaCoverageThresholds(FVector4& AttributeValue) const
-	{
-		IMPLEMENT_NODE_ATTRIBUTE_GETTER(AlphaCoverageThresholds, FVector4);
-	}
-
-	UFUNCTION(BlueprintCallable, Category = "Interchange | Node | Texture")
-	bool SetCustomAlphaCoverageThresholds(const FVector4& AttributeValue, bool bAddApplyDelegate = true)
-	{
-#if WITH_EDITORONLY_DATA
-		IMPLEMENT_NODE_ATTRIBUTE_SETTER(UInterchangeTextureNode, AlphaCoverageThresholds, FVector4, UTexture)
-#else
-		IMPLEMENT_NODE_ATTRIBUTE_SETTER_NODELEGATE(AlphaCoverageThresholds, FVector4)
-#endif
-	}
-
-	/** Return false if the Attribute was not set previously.*/
-	UFUNCTION(BlueprintCallable, Category = "Interchange | Node | Texture")
-	bool GetCustombFlipGreenChannel(bool& AttributeValue) const
-	{
-		IMPLEMENT_NODE_ATTRIBUTE_GETTER(bFlipGreenChannel, bool);
-	}
-
-	UFUNCTION(BlueprintCallable, Category = "Interchange | Node | Texture")
-	bool SetCustombFlipGreenChannel(const bool& AttributeValue, bool bAddApplyDelegate = true)
-	{
-#if WITH_EDITORONLY_DATA
-		IMPLEMENT_NODE_ATTRIBUTE_SETTER(UInterchangeTextureNode, bFlipGreenChannel, bool, UTexture)
-#else
-		IMPLEMENT_NODE_ATTRIBUTE_SETTER_NODELEGATE(bFlipGreenChannel, bool)
-#endif
-	}
 
-	/** Return false if the Attribute was not set previously.*/
-	UFUNCTION(BlueprintCallable, Category = "Interchange | Node | Texture")
-	bool GetCustomPowerOfTwoMode(uint8& AttributeValue) const
-	{
-		IMPLEMENT_NODE_ATTRIBUTE_GETTER(PowerOfTwoMode, uint8);
-	}
-
-	UFUNCTION(BlueprintCallable, Category = "Interchange | Node | Texture")
-	bool SetCustomPowerOfTwoMode(const uint8& AttributeValue, bool bAddApplyDelegate = true)
-	{
-#if WITH_EDITORONLY_DATA
-		IMPLEMENT_NODE_ATTRIBUTE_SETTER(UInterchangeTextureNode, PowerOfTwoMode, uint8, UTexture)
-#else
-		IMPLEMENT_NODE_ATTRIBUTE_SETTER_NODELEGATE(PowerOfTwoMode, uint8)
-#endif
-	}
-
-	/** Return false if the Attribute was not set previously.*/
-	UFUNCTION(BlueprintCallable, Category = "Interchange | Node | Texture")
-	bool GetCustomPaddingColor(FColor& AttributeValue) const
-	{
-		IMPLEMENT_NODE_ATTRIBUTE_GETTER(PaddingColor, FColor);
-	}
-
-	UFUNCTION(BlueprintCallable, Category = "Interchange | Node | Texture")
-	bool SetCustomPaddingColor(const FColor& AttributeValue, bool bAddApplyDelegate = true)
-	{
-#if WITH_EDITORONLY_DATA
-		IMPLEMENT_NODE_ATTRIBUTE_SETTER(UInterchangeTextureNode, PaddingColor, FColor, UTexture)
-#else
-		IMPLEMENT_NODE_ATTRIBUTE_SETTER_NODELEGATE(PaddingColor, FColor)
-#endif
-	}
-
-	/** Return false if the Attribute was not set previously.*/
-	UFUNCTION(BlueprintCallable, Category = "Interchange | Node | Texture")
-	bool GetCustomFilter(uint8& AttributeValue) const
-	{
-		IMPLEMENT_NODE_ATTRIBUTE_GETTER(Filter, uint8);
-	}
-
-	UFUNCTION(BlueprintCallable, Category = "Interchange | Node | Texture")
-	bool SetCustomFilter(const uint8& AttributeValue, bool bAddApplyDelegate = true)
-	{
-		IMPLEMENT_NODE_ATTRIBUTE_SETTER(UInterchangeTextureNode, Filter, uint8, UTexture)
-	}
-
-	/** Return false if the Attribute was not set previously.*/
-	UFUNCTION(BlueprintCallable, Category = "Interchange | Node | Texture")
-	bool GetCustomMipLoadOptions(uint8& AttributeValue) const
-	{
-		IMPLEMENT_NODE_ATTRIBUTE_GETTER(MipLoadOptions, uint8);
-	}
-
-	UFUNCTION(BlueprintCallable, Category = "Interchange | Node | Texture")
-	bool SetCustomMipLoadOptions(const uint8& AttributeValue, bool bAddApplyDelegate = true)
-	{
-		IMPLEMENT_NODE_ATTRIBUTE_SETTER(UInterchangeTextureNode, MipLoadOptions, uint8, UTexture)
-	}
-
-	/** Return false if the Attribute was not set previously.*/
-	UFUNCTION(BlueprintCallable, Category = "Interchange | Node | Texture")
-	bool GetCustomSRGB(bool& AttributeValue) const
-	{
-		IMPLEMENT_NODE_ATTRIBUTE_GETTER(SRGB, bool);
-	}
-
-	UFUNCTION(BlueprintCallable, Category = "Interchange | Node | Texture")
-	bool SetCustomSRGB(const bool& AttributeValue, bool bAddApplyDelegate = true)
-	{
-		IMPLEMENT_NODE_ATTRIBUTE_SETTER(UInterchangeTextureNode, SRGB, bool, UTexture)
-	}
-
-	/** Return false if the Attribute was not set previously.*/
-	UFUNCTION(BlueprintCallable, Category = "Interchange | Node | Texture")
-	bool GetCustombUseLegacyGamma(bool& AttributeValue) const
-	{
-		IMPLEMENT_NODE_ATTRIBUTE_GETTER(bUseLegacyGamma, bool);
-	}
-
-	UFUNCTION(BlueprintCallable, Category = "Interchange | Node | Texture")
-	bool SetCustombUseLegacyGamma(const bool& AttributeValue, bool bAddApplyDelegate = true)
-	{
-#if WITH_EDITORONLY_DATA
-		IMPLEMENT_NODE_ATTRIBUTE_SETTER(UInterchangeTextureNode, bUseLegacyGamma, bool, UTexture)
-#else
-		IMPLEMENT_NODE_ATTRIBUTE_SETTER_NODELEGATE(bUseLegacyGamma, bool)
-#endif
-	}
-
-	/** Return false if the Attribute was not set previously.*/
-	UFUNCTION(BlueprintCallable, Category = "Interchange | Node | Texture")
-	bool GetCustomVirtualTextureStreaming(bool& AttributeValue) const
-	{
-		IMPLEMENT_NODE_ATTRIBUTE_GETTER(VirtualTextureStreaming, bool);
-	}
-
-	UFUNCTION(BlueprintCallable, Category = "Interchange | Node | Texture")
-	bool SetCustomVirtualTextureStreaming(const bool& AttributeValue, bool bAddApplyDelegate = true)
-	{
-		IMPLEMENT_NODE_ATTRIBUTE_SETTER(UInterchangeTextureNode, VirtualTextureStreaming, bool, UTexture)
-	}
-
-	/** Return false if the Attribute was not set previously.*/
-	UFUNCTION(BlueprintCallable, Category = "Interchange | Node | Texture")
-	bool GetCustomAddressX(uint8& AttributeValue) const
-	{
-		IMPLEMENT_NODE_ATTRIBUTE_GETTER(AddressX, uint8);
-	}
-
-	UFUNCTION(BlueprintCallable, Category = "Interchange | Node | Texture")
-	bool SetCustomAddressX(const uint8& AttributeValue, bool bAddApplyDelegate = true)
-	{
-		IMPLEMENT_NODE_ATTRIBUTE_SETTER(UInterchangeTextureNode, AddressX, uint8, UTexture2D)
-	}
-
-	/** Return false if the Attribute was not set previously.*/
-	UFUNCTION(BlueprintCallable, Category = "Interchange | Node | Texture")
-	bool GetCustomAddressY(uint8& AttributeValue) const
-	{
-		IMPLEMENT_NODE_ATTRIBUTE_GETTER(AddressY, uint8);
-	}
-
-	UFUNCTION(BlueprintCallable, Category = "Interchange | Node | Texture")
-	bool SetCustomAddressY(const uint8& AttributeValue, bool bAddApplyDelegate = true)
-	{
-		IMPLEMENT_NODE_ATTRIBUTE_SETTER(UInterchangeTextureNode, AddressY, uint8, UTexture2D)
-	}
-
-	//////////////////////////////////////////////////////////////////////////
-	//Level of Detail
-
-	/** Return false if the Attribute was not set previously.*/
-	UFUNCTION(BlueprintCallable, Category = "Interchange | Node | Texture")
-	bool GetCustombPreserveBorder(bool& AttributeValue) const
-	{
-		IMPLEMENT_NODE_ATTRIBUTE_GETTER(bPreserveBorder, bool);
-	}
-
-	UFUNCTION(BlueprintCallable, Category = "Interchange | Node | Texture")
-	bool SetCustombPreserveBorder(const bool& AttributeValue, bool bAddApplyDelegate = true)
-	{
-#if WITH_EDITORONLY_DATA
-		IMPLEMENT_NODE_ATTRIBUTE_SETTER(UInterchangeTextureNode, bPreserveBorder, bool, UTexture)
-#else
-		IMPLEMENT_NODE_ATTRIBUTE_SETTER_NODELEGATE(bPreserveBorder, bool)
-#endif
-	}
-
-	/** Return false if the Attribute was not set previously.*/
-	UFUNCTION(BlueprintCallable, Category = "Interchange | Node | Texture")
-	bool GetCustomMipGenSettings(uint8& AttributeValue) const
-	{
-		IMPLEMENT_NODE_ATTRIBUTE_GETTER(MipGenSettings, uint8);
-	}
-
-	UFUNCTION(BlueprintCallable, Category = "Interchange | Node | Texture")
-	bool SetCustomMipGenSettings(const uint8& AttributeValue, bool bAddApplyDelegate = true)
-	{
-#if WITH_EDITORONLY_DATA
-		IMPLEMENT_NODE_ATTRIBUTE_SETTER(UInterchangeTextureNode, MipGenSettings, uint8, UTexture)
-#else
-		IMPLEMENT_NODE_ATTRIBUTE_SETTER_NODELEGATE(MipGenSettings, uint8)
-#endif
-	}
-
-	/** Return false if the Attribute was not set previously.*/
-	UFUNCTION(BlueprintCallable, Category = "Interchange | Node | Texture")
-	bool GetCustomLODBias(int32& AttributeValue) const
-	{
-		IMPLEMENT_NODE_ATTRIBUTE_GETTER(LODBias, int32);
-	}
-
-	UFUNCTION(BlueprintCallable, Category = "Interchange | Node | Texture")
-	bool SetCustomLODBias(const int32& AttributeValue, bool bAddApplyDelegate = true)
-	{
-		IMPLEMENT_NODE_ATTRIBUTE_SETTER(UInterchangeTextureNode, LODBias, int32, UTexture)
-	}
-
-	/** Return false if the Attribute was not set previously.*/
-	UFUNCTION(BlueprintCallable, Category = "Interchange | Node | Texture")
-	bool GetCustomLODGroup(uint8& AttributeValue) const
-	{
-		IMPLEMENT_NODE_ATTRIBUTE_GETTER(LODGroup, uint8);
-	}
-
-	UFUNCTION(BlueprintCallable, Category = "Interchange | Node | Texture")
-	bool SetCustomLODGroup(const uint8& AttributeValue, bool bAddApplyDelegate = true)
-	{
-		IMPLEMENT_NODE_ATTRIBUTE_SETTER(UInterchangeTextureNode, LODGroup, uint8, UTexture)
-	}
-
-	//TODO support per platform data in the FAttributeStorage, so we can set different value per platform at the pipeline stage, We set only the default value for now
-	/** Return false if the Attribute was not set previously.*/
-	UFUNCTION(BlueprintCallable, Category = "Interchange | Node | Texture")
-	bool GetCustomDownscale(float& AttributeValue) const
-	{
-		IMPLEMENT_NODE_ATTRIBUTE_GETTER(Downscale, float);
-	}
-
-	UFUNCTION(BlueprintCallable, Category = "Interchange | Node | Texture")
-	bool SetCustomDownscale(const float& AttributeValue, bool bAddApplyDelegate = true)
-	{
-		IMPLEMENT_NODE_ATTRIBUTE_SETTER(UInterchangeTextureNode, Downscale, float, UTexture)
-	}
-
-	/** Return false if the Attribute was not set previously.*/
-	UFUNCTION(BlueprintCallable, Category = "Interchange | Node | Texture")
-		bool GetCustomDownscaleOptions(uint8& AttributeValue) const
-	{
-		IMPLEMENT_NODE_ATTRIBUTE_GETTER(DownscaleOptions, uint8);
-	}
-
-	UFUNCTION(BlueprintCallable, Category = "Interchange | Node | Texture")
-		bool SetCustomDownscaleOptions(const uint8& AttributeValue, bool bAddApplyDelegate = true)
-	{
-		IMPLEMENT_NODE_ATTRIBUTE_SETTER(UInterchangeTextureNode, DownscaleOptions, uint8, UTexture)
-	}
-
-	//////////////////////////////////////////////////////////////////////////
-	//Compositing
-
-	/** Return false if the Attribute was not set previously.*/
-	UFUNCTION(BlueprintCallable, Category = "Interchange | Node | Texture")
-	bool GetCustomCompositeTextureMode(uint8& AttributeValue) const
-	{
-		IMPLEMENT_NODE_ATTRIBUTE_GETTER(CompositeTextureMode, uint8);
-	}
-
-	UFUNCTION(BlueprintCallable, Category = "Interchange | Node | Texture")
-	bool SetCustomCompositeTextureMode(const uint8& AttributeValue, bool bAddApplyDelegate = true)
-	{
-#if WITH_EDITORONLY_DATA
-		IMPLEMENT_NODE_ATTRIBUTE_SETTER(UInterchangeTextureNode, CompositeTextureMode, uint8, UTexture)
-#else
-		IMPLEMENT_NODE_ATTRIBUTE_SETTER_NODELEGATE(CompositeTextureMode, uint8)
-#endif
-	}
-
-	/** Return false if the Attribute was not set previously.*/
-	UFUNCTION(BlueprintCallable, Category = "Interchange | Node | Texture")
-	bool GetCustomCompositePower(float& AttributeValue) const
-	{
-		IMPLEMENT_NODE_ATTRIBUTE_GETTER(CompositePower, float);
-	}
-
-	UFUNCTION(BlueprintCallable, Category = "Interchange | Node | Texture")
-	bool SetCustomCompositePower(const float& AttributeValue, bool bAddApplyDelegate = true)
-	{
-#if WITH_EDITORONLY_DATA
-		IMPLEMENT_NODE_ATTRIBUTE_SETTER(UInterchangeTextureNode, CompositePower, float, UTexture)
-#else
-		IMPLEMENT_NODE_ATTRIBUTE_SETTER_NODELEGATE(CompositePower, float)
-#endif
-	}
-
-private:
-
-	void FillAssetClassFromAttribute()
-	{
-#if WITH_ENGINE
-		FString OperationName = GetTypeName() + TEXT(".GetAssetClassName");
-		FString ClassName;
-		InterchangePrivateNodeBase::GetCustomAttribute<FString>(*Attributes, ClassNameAttributeKey, OperationName, ClassName);
-		if (ClassName.Equals(UTexture2D::StaticClass()->GetName()))
-		{
-			AssetClass = UTexture2D::StaticClass();
-			bIsTextureNodeClassInitialized = true;
-		}
-		else if (ClassName.Equals(UTexture::StaticClass()->GetName()))
-		{
-			AssetClass = UTexture::StaticClass();
-			bIsTextureNodeClassInitialized = true;
-		}
-#endif
-	}
-
-	bool SetTextureNodeClassFromClassAttribute()
-	{
-		if (!bIsTextureNodeClassInitialized)
-		{
-			FillAssetClassFromAttribute();
-		}
-		return bIsTextureNodeClassInitialized;
-	}
-
-	bool IsEditorOnlyDataDefined()
-	{
-#if WITH_EDITORONLY_DATA
-		return true;
-#else
-		return false;
-#endif
-	}
-
-	const UE::Interchange::FAttributeKey ClassNameAttributeKey = UE::Interchange::FBaseNodeStaticData::ClassTypeAttributeKey();
-
-	//Texture Adjustments
-	const UE::Interchange::FAttributeKey Macro_CustomAdjustBrightnessKey = UE::Interchange::FAttributeKey(TEXT("AdjustBrightness"));
-	const UE::Interchange::FAttributeKey Macro_CustomAdjustBrightnessCurveKey = UE::Interchange::FAttributeKey(TEXT("AdjustBrightnessCurve"));
-	const UE::Interchange::FAttributeKey Macro_CustomAdjustVibranceKey = UE::Interchange::FAttributeKey(TEXT("AdjustVibrance"));
-	const UE::Interchange::FAttributeKey Macro_CustomAdjustSaturationKey = UE::Interchange::FAttributeKey(TEXT("AdjustSaturation"));
-	const UE::Interchange::FAttributeKey Macro_CustomAdjustRGBCurveKey = UE::Interchange::FAttributeKey(TEXT("AdjustRGBCurve"));
-	const UE::Interchange::FAttributeKey Macro_CustomAdjustHueKey = UE::Interchange::FAttributeKey(TEXT("AdjustHue"));
-	const UE::Interchange::FAttributeKey Macro_CustomAdjustMinAlphaKey = UE::Interchange::FAttributeKey(TEXT("AdjustMinAlpha"));
-	const UE::Interchange::FAttributeKey Macro_CustomAdjustMaxAlphaKey = UE::Interchange::FAttributeKey(TEXT("AdjustMaxAlpha"));
-	const UE::Interchange::FAttributeKey Macro_CustombChromaKeyTextureKey = UE::Interchange::FAttributeKey(TEXT("bChromaKeyTexture"));
-	const UE::Interchange::FAttributeKey Macro_CustomChromaKeyThresholdKey = UE::Interchange::FAttributeKey(TEXT("ChromaKeyThreshold"));
-	const UE::Interchange::FAttributeKey Macro_CustomChromaKeyColorKey = UE::Interchange::FAttributeKey(TEXT("ChromaKeyColor"));
-
-	//Texture Compression
-	const UE::Interchange::FAttributeKey Macro_CustomCompressionNoAlphaKey = UE::Interchange::FAttributeKey(TEXT("CompressionNoAlpha"));
-	const UE::Interchange::FAttributeKey Macro_CustomDeferCompressionKey = UE::Interchange::FAttributeKey(TEXT("DeferCompression"));
-	const UE::Interchange::FAttributeKey Macro_CustomLossyCompressionAmountKey = UE::Interchange::FAttributeKey(TEXT("LossyCompressionAmount"));
-	const UE::Interchange::FAttributeKey Macro_CustomMaxTextureSizeKey = UE::Interchange::FAttributeKey(TEXT("MaxTextureSize"));
-	const UE::Interchange::FAttributeKey Macro_CustomCompressionQualityKey = UE::Interchange::FAttributeKey(TEXT("CompressionQuality"));
-	const UE::Interchange::FAttributeKey Macro_CustomCompressionSettingsKey = UE::Interchange::FAttributeKey(TEXT("CompressionSettings"));
-
-	//Texture general
-	const UE::Interchange::FAttributeKey Macro_CustombDitherMipMapAlphaKey = UE::Interchange::FAttributeKey(TEXT("bDitherMipMapAlpha"));
-	const UE::Interchange::FAttributeKey Macro_CustomAlphaCoverageThresholdsKey = UE::Interchange::FAttributeKey(TEXT("AlphaCoverageThresholds"));
-	const UE::Interchange::FAttributeKey Macro_CustombFlipGreenChannelKey = UE::Interchange::FAttributeKey(TEXT("bFlipGreenChannel"));
-	const UE::Interchange::FAttributeKey Macro_CustombForcePVRTC4Key = UE::Interchange::FAttributeKey(TEXT("bForcePVRTC4"));
-	const UE::Interchange::FAttributeKey Macro_CustomPowerOfTwoModeKey = UE::Interchange::FAttributeKey(TEXT("PowerOfTwoMode"));
-	const UE::Interchange::FAttributeKey Macro_CustomPaddingColorKey = UE::Interchange::FAttributeKey(TEXT("PaddingColor"));
-	const UE::Interchange::FAttributeKey Macro_CustomFilterKey = UE::Interchange::FAttributeKey(TEXT("Filter"));
-	const UE::Interchange::FAttributeKey Macro_CustomMipLoadOptionsKey = UE::Interchange::FAttributeKey(TEXT("MipLoadOptions"));
-	const UE::Interchange::FAttributeKey Macro_CustomSRGBKey = UE::Interchange::FAttributeKey(TEXT("SRGB"));
-	const UE::Interchange::FAttributeKey Macro_CustombUseLegacyGammaKey = UE::Interchange::FAttributeKey(TEXT("bUseLegacyGamma"));
-	const UE::Interchange::FAttributeKey Macro_CustomVirtualTextureStreamingKey = UE::Interchange::FAttributeKey(TEXT("VirtualTextureStreaming"));
-	const UE::Interchange::FAttributeKey Macro_CustomAddressXKey = UE::Interchange::FAttributeKey(TEXT("AddressX"));
-	const UE::Interchange::FAttributeKey Macro_CustomAddressYKey = UE::Interchange::FAttributeKey(TEXT("AddressY"));
-
-	//Level of Detail
-	const UE::Interchange::FAttributeKey Macro_CustombPreserveBorderKey = UE::Interchange::FAttributeKey(TEXT("bPreserveBorder"));
-	const UE::Interchange::FAttributeKey Macro_CustomMipGenSettingsKey = UE::Interchange::FAttributeKey(TEXT("MipGenSettings"));
-	const UE::Interchange::FAttributeKey Macro_CustomLODBiasKey = UE::Interchange::FAttributeKey(TEXT("LODBias"));
-	const UE::Interchange::FAttributeKey Macro_CustomLODGroupKey = UE::Interchange::FAttributeKey(TEXT("LODGroup"));
-	//TODO support per platform data in the FAttributeStorage, so we can set different value per platform at the pipeline stage, We set only the default value for now
-	const UE::Interchange::FAttributeKey Macro_CustomDownscaleKey = UE::Interchange::FAttributeKey(TEXT("Downscale"));
-	const UE::Interchange::FAttributeKey Macro_CustomDownscaleOptionsKey = UE::Interchange::FAttributeKey(TEXT("DownscaleOptions"));
-
-	//Compositing
-	const UE::Interchange::FAttributeKey Macro_CustomCompositeTextureModeKey = UE::Interchange::FAttributeKey(TEXT("CompositeTextureMode"));
-	const UE::Interchange::FAttributeKey Macro_CustomCompositePowerKey = UE::Interchange::FAttributeKey(TEXT("CompositePower"));
-
-#if WITH_EDITORONLY_DATA
-	IMPLEMENT_NODE_ATTRIBUTE_APPLY_UOBJECT(AdjustBrightness, float, UTexture, );
-	IMPLEMENT_NODE_ATTRIBUTE_APPLY_UOBJECT(AdjustBrightnessCurve, float, UTexture, );
-	IMPLEMENT_NODE_ATTRIBUTE_APPLY_UOBJECT(AdjustVibrance, float, UTexture, );
-	IMPLEMENT_NODE_ATTRIBUTE_APPLY_UOBJECT(AdjustSaturation, float, UTexture, );
-	IMPLEMENT_NODE_ATTRIBUTE_APPLY_UOBJECT(AdjustRGBCurve, float, UTexture, );
-	IMPLEMENT_NODE_ATTRIBUTE_APPLY_UOBJECT(AdjustHue, float, UTexture, );
-	IMPLEMENT_NODE_ATTRIBUTE_APPLY_UOBJECT(AdjustMinAlpha, float, UTexture, );
-	IMPLEMENT_NODE_ATTRIBUTE_APPLY_UOBJECT(AdjustMaxAlpha, float, UTexture, );
-	IMPLEMENT_NODE_ATTRIBUTE_APPLY_UOBJECT(bChromaKeyTexture, bool, UTexture, );
-	IMPLEMENT_NODE_ATTRIBUTE_APPLY_UOBJECT(ChromaKeyThreshold, float, UTexture, );
-	IMPLEMENT_NODE_ATTRIBUTE_APPLY_UOBJECT(ChromaKeyColor, FColor, UTexture, );
-
-	//Texture Compression
-	IMPLEMENT_NODE_ATTRIBUTE_APPLY_UOBJECT(CompressionNoAlpha, bool, UTexture, );
-	IMPLEMENT_NODE_ATTRIBUTE_APPLY_UOBJECT(DeferCompression, bool, UTexture, );
-	IMPLEMENT_NODE_ATTRIBUTE_APPLY_UOBJECT(LossyCompressionAmount, uint8, UTexture, TEnumAsByte<ETextureLossyCompressionAmount>);
-	IMPLEMENT_NODE_ATTRIBUTE_APPLY_UOBJECT(MaxTextureSize, int32, UTexture, );
-	IMPLEMENT_NODE_ATTRIBUTE_APPLY_UOBJECT(CompressionQuality, uint8, UTexture, TEnumAsByte<enum ETextureCompressionQuality>);
-#endif
-	IMPLEMENT_NODE_ATTRIBUTE_APPLY_UOBJECT(CompressionSettings, uint8, UTexture, TEnumAsByte<enum TextureCompressionSettings>);
-#if WITH_EDITORONLY_DATA
-	//Texture general
-	IMPLEMENT_NODE_ATTRIBUTE_APPLY_UOBJECT(bDitherMipMapAlpha, bool, UTexture, );
-	IMPLEMENT_NODE_ATTRIBUTE_APPLY_UOBJECT(AlphaCoverageThresholds, FVector4, UTexture, );
-	IMPLEMENT_NODE_ATTRIBUTE_APPLY_UOBJECT(bFlipGreenChannel, bool, UTexture, );
-	IMPLEMENT_NODE_ATTRIBUTE_APPLY_UOBJECT(PowerOfTwoMode, uint8, UTexture, TEnumAsByte<enum ETexturePowerOfTwoSetting::Type>);
-	IMPLEMENT_NODE_ATTRIBUTE_APPLY_UOBJECT(PaddingColor, FColor, UTexture, );
-#endif //#if WITH_EDITORONLY_DATA
-	IMPLEMENT_NODE_ATTRIBUTE_APPLY_UOBJECT(Filter, uint8, UTexture, TEnumAsByte<enum TextureFilter>);
-	IMPLEMENT_NODE_ATTRIBUTE_APPLY_UOBJECT(MipLoadOptions, uint8, UTexture, ETextureMipLoadOptions);
-	IMPLEMENT_NODE_ATTRIBUTE_APPLY_UOBJECT(SRGB, bool, UTexture, );
-#if WITH_EDITORONLY_DATA
-	IMPLEMENT_NODE_ATTRIBUTE_APPLY_UOBJECT(bUseLegacyGamma, bool, UTexture, );
-#endif //#if WITH_EDITORONLY_DATA
-	IMPLEMENT_NODE_ATTRIBUTE_APPLY_UOBJECT(VirtualTextureStreaming, bool, UTexture, );
-	IMPLEMENT_NODE_ATTRIBUTE_APPLY_UOBJECT(AddressX, uint8, UTexture2D, TEnumAsByte<enum TextureAddress>);
-	IMPLEMENT_NODE_ATTRIBUTE_APPLY_UOBJECT(AddressY, uint8, UTexture2D, TEnumAsByte<enum TextureAddress>);
-#if WITH_EDITORONLY_DATA
-	//Level of Detail
-	IMPLEMENT_NODE_ATTRIBUTE_APPLY_UOBJECT(bPreserveBorder, bool, UTexture, );
-	IMPLEMENT_NODE_ATTRIBUTE_APPLY_UOBJECT(MipGenSettings, uint8, UTexture, TEnumAsByte<enum TextureMipGenSettings>);
-#endif //#if WITH_EDITORONLY_DATA
-	IMPLEMENT_NODE_ATTRIBUTE_APPLY_UOBJECT(LODBias, int32, UTexture, );
-	IMPLEMENT_NODE_ATTRIBUTE_APPLY_UOBJECT(LODGroup, uint8, UTexture, TEnumAsByte<enum TextureGroup>);
-
-	//TODO support per platform data in the FAttributeStorage, so we can set different value per platform at the pipeline stage, We set only the default value for now
-	//IMPLEMENT_NODE_ATTRIBUTE_APPLY_UOBJECT(Downscale, float, UTexture, );
-#if WITH_ENGINE
-	bool ApplyCustomDownscaleToAsset(UObject * Asset) const
-	{
-		if (!Asset)
-		{
-			return false;
-		}
-		UTexture* TypedObject = Cast<UTexture>(Asset);
-		if (!TypedObject)
-		{
-			return false;
-		}
-		float ValueData;
-		if (GetCustomDownscale(ValueData))
-		{
-			TypedObject->Downscale.Default = ValueData;
-			return true;
-		}
-		return false;
-	}
-
-	bool FillCustomDownscaleFromAsset(UObject* Asset)
-	{
-		if (!Asset)
-		{
-			return false;
-		}
-		UTexture* TypedObject = Cast<UTexture>(Asset);
-		if (!TypedObject)
-		{
-			return false;
-		}
-		if (SetCustomDownscale(TypedObject->Downscale.Default, false))
-		{
-			return true;
-		}
-		return false;
-	}
-#endif //WITH_ENGINE
-
-	IMPLEMENT_NODE_ATTRIBUTE_APPLY_UOBJECT(DownscaleOptions, uint8, UTexture, ETextureDownscaleOptions);
-#if WITH_EDITORONLY_DATA
-	//Compositing
-	IMPLEMENT_NODE_ATTRIBUTE_APPLY_UOBJECT(CompositeTextureMode, uint8, UTexture, TEnumAsByte<enum ECompositeTextureMode>);
-	IMPLEMENT_NODE_ATTRIBUTE_APPLY_UOBJECT(CompositePower, float, UTexture, );
-#endif //#if WITH_EDITORONLY_DATA
 protected:
-#if WITH_ENGINE
-	TSubclassOf<UTexture> AssetClass = nullptr;
-#endif
-
-	UE::Interchange::TMapAttributeHelper<int32,FString> SourceBlocks;
 
-	bool bIsTextureNodeClassInitialized = false;
+	UE::Interchange::TMapAttributeHelper<int32, FString> SourceBlocks;
 };
diff --git a/Engine/Plugins/Runtime/Interchange/Source/Parsers/Fbx/Private/FbxAPI.cpp b/Engine/Plugins/Runtime/Interchange/Source/Parsers/Fbx/Private/FbxAPI.cpp
--- a/Engine/Plugins/Runtime/Interchange/Source/Parsers/Fbx/Private/FbxAPI.cpp
+++ b/Engine/Plugins/Runtime/Interchange/Source/Parsers/Fbx/Private/FbxAPI.cpp
@@ -7,9 +7,8 @@
 #include "FbxHelper.h"
 #include "FbxInclude.h"
 #include "FbxMaterial.h"
+#include "FbxMesh.h"
 #include "FbxScene.h"
-#include "FbxSkeletalMesh.h"
-#include "FbxSkeleton.h"
 #include "InterchangeMaterialNode.h"
 #include "InterchangeTextureNode.h"
 #include "Nodes/InterchangeBaseNodeContainer.h"
@@ -92,10 +91,11 @@ namespace UE
 
 			void FFbxParser::FillContainerWithFbxScene(UInterchangeBaseNodeContainer& NodeContainer, TArray<FString>& JSonErrorMessages)
 			{
-				FFbxMaterial::AddAllSceneMaterials(SDKScene, NodeContainer, JSonErrorMessages);
-				FFbxSkeleton::AddAllSceneSkeletons(SDKScene, NodeContainer, JSonErrorMessages, SourceFilename);
-				FFbxSkeletalMesh::AddAllSceneSkeletalMeshes(SDKScene, SDKGeometryConverter, NodeContainer, JSonErrorMessages, PayloadContexts, SourceFilename);
-				//FFbxScene::AddHierarchy(SDKScene, NodeContainer, JSonErrorMessages);
+				CleanupFbxData(JSonErrorMessages);
+				FFbxMaterial::AddAllTextures(SDKScene, NodeContainer, JSonErrorMessages);
+				FFbxMaterial::AddAllMaterials(SDKScene, NodeContainer, JSonErrorMessages);
+				FFbxMesh::AddAllMeshes(SDKScene, SDKGeometryConverter, NodeContainer, JSonErrorMessages, PayloadContexts);
+				FFbxScene::AddHierarchy(SDKScene, NodeContainer, JSonErrorMessages);
 			}
 
 			bool FFbxParser::FetchPayloadData(const FString& PayloadKey, const FString& PayloadFilepath, TArray<FString>& JSonErrorMessages)
@@ -110,6 +110,25 @@ namespace UE
 				return PayloadContext->FetchPayloadToFile(PayloadFilepath, JSonErrorMessages);
 			}
 
+			void FFbxParser::CleanupFbxData(TArray<FString>& JSonErrorMessages)
+			{
+				//////////////////////////////////////////////////////////////////////////
+				// Make sure there is a valid bind pose
+
+				//Find root bones
+				const int32 Default_NbPoses = SDKScene->GetFbxManager()->GetBindPoseCount(SDKScene);
+				// If there are no BindPoses, the following will generate them.
+				SDKScene->GetFbxManager()->CreateMissingBindPoses(SDKScene);
+				//if we created missing bind poses, update the number of bind poses
+				const int32 NbPoses = SDKScene->GetFbxManager()->GetBindPoseCount(SDKScene);
+				if (NbPoses != Default_NbPoses)
+				{
+					JSonErrorMessages.Add(TEXT("{\"Msg\" : {\"Type\" : \"Warning\",\n\"Msg\" : \"FbxParser: Missing bind pose, the fbx sdk create one.\"}}"));
+				}
+
+				//////////////////////////////////////////////////////////////////////////
+				// Scene conversion
+			}
 		} //ns Private
 	} //ns Interchange
 } //ns UE
diff --git a/Engine/Plugins/Runtime/Interchange/Source/Parsers/Fbx/Private/FbxAPI.h b/Engine/Plugins/Runtime/Interchange/Source/Parsers/Fbx/Private/FbxAPI.h
--- a/Engine/Plugins/Runtime/Interchange/Source/Parsers/Fbx/Private/FbxAPI.h
+++ b/Engine/Plugins/Runtime/Interchange/Source/Parsers/Fbx/Private/FbxAPI.h
@@ -41,6 +41,10 @@ namespace UE
 				bool FetchPayloadData(const FString& PayloadKey, const FString& PayloadFilepath, TArray<FString>& JSonErrorMessages);
 			
 			private:
+
+				void CleanupFbxData(TArray<FString>& JSonErrorMessages);
+
+
 				FbxManager* SDKManager = nullptr;
 				FbxScene* SDKScene = nullptr;
 				FbxImporter* SDKImporter = nullptr;
diff --git a/Engine/Plugins/Runtime/Interchange/Source/Parsers/Fbx/Private/FbxMaterial.cpp b/Engine/Plugins/Runtime/Interchange/Source/Parsers/Fbx/Private/FbxMaterial.cpp
--- a/Engine/Plugins/Runtime/Interchange/Source/Parsers/Fbx/Private/FbxMaterial.cpp
+++ b/Engine/Plugins/Runtime/Interchange/Source/Parsers/Fbx/Private/FbxMaterial.cpp
@@ -7,6 +7,7 @@
 #include "FbxHelper.h"
 #include "FbxInclude.h"
 #include "InterchangeMaterialNode.h"
+#include "InterchangeSceneNode.h"
 #include "InterchangeTextureNode.h"
 #include "Nodes/InterchangeBaseNodeContainer.h"
 
@@ -16,12 +17,8 @@ namespace UE
 	{
 		namespace Private
 		{
-			UInterchangeMaterialNode* FFbxMaterial::CreateMaterialNode(UInterchangeBaseNodeContainer& NodeContainer, const FString& NodeName, TArray<FString>& JSonErrorMessages)
+			UInterchangeMaterialNode* FFbxMaterial::CreateMaterialNode(UInterchangeBaseNodeContainer& NodeContainer, const FString& NodeUid, const FString& NodeName, TArray<FString>& JSonErrorMessages)
 			{
-				FString NodeUID(NodeName);
-				FString MaterialNameNoSkin = UE::Interchange::Material::RemoveSkinFromName(NodeName);
-
-				FString DisplayLabel = MaterialNameNoSkin;
 				UInterchangeMaterialNode* MaterialNode = NewObject<UInterchangeMaterialNode>(&NodeContainer, NAME_None);
 				if (!ensure(MaterialNode))
 				{
@@ -29,16 +26,15 @@ namespace UE
 					return nullptr;
 				}
 				// Creating a UMaterialInterface
-				MaterialNode->InitializeMaterialNode(NodeUID, DisplayLabel, TEXT("Material"));
-				MaterialNode->SetPayLoadKey(NodeName);
+				MaterialNode->InitializeNode(NodeUid, NodeName, EInterchangeNodeContainerType::NodeContainerType_TranslatedAsset);
+				MaterialNode->SetPayLoadKey(NodeUid);
 				NodeContainer.AddNode(MaterialNode);
 				return MaterialNode;
 			}
 
-			UInterchangeTextureNode* FFbxMaterial::CreateTexture2DNode(UInterchangeBaseNodeContainer& NodeContainer, const FString& TextureFilePath, TArray<FString>& JSonErrorMessages)
+			UInterchangeTextureNode* FFbxMaterial::CreateTexture2DNode(UInterchangeBaseNodeContainer& NodeContainer, const FString& NodeUid, const FString& TextureFilePath, TArray<FString>& JSonErrorMessages)
 			{
 				FString DisplayLabel = FPaths::GetBaseFilename(TextureFilePath);
-				FString NodeUID(TextureFilePath);
 				UInterchangeTextureNode* TextureNode = NewObject<UInterchangeTextureNode>(&NodeContainer, NAME_None);
 				if (!ensure(TextureNode))
 				{
@@ -46,8 +42,11 @@ namespace UE
 					return nullptr;
 				}
 				// Creating a UTexture2D
-				TextureNode->InitializeTextureNode(NodeUID, DisplayLabel, TEXT("Texture2D"));
-				TextureNode->SetPayLoadKey(TextureFilePath);
+				TextureNode->InitializeNode(NodeUid, DisplayLabel, EInterchangeNodeContainerType::NodeContainerType_TranslatedAsset);
+				//All texture translator expect a file has the payload key
+				FString NormalizeFilePath = TextureFilePath;
+				FPaths::NormalizeFilename(NormalizeFilePath);
+				TextureNode->SetPayLoadKey(NormalizeFilePath);
 				NodeContainer.AddNode(TextureNode);
 				return TextureNode;
 			}
@@ -56,11 +55,11 @@ namespace UE
 			{
 				//Create a material node
 				FString MaterialName = FFbxHelper::GetFbxObjectName(SurfaceMaterial);
-				FString NodeUID(MaterialName);
-				UInterchangeMaterialNode* MaterialNode = Cast<UInterchangeMaterialNode>(NodeContainer.GetNode(NodeUID));
+				FString NodeUid = TEXT("\\Material\\") + MaterialName;
+				UInterchangeMaterialNode* MaterialNode = Cast<UInterchangeMaterialNode>(NodeContainer.GetNode(NodeUid));
 				if (!MaterialNode)
 				{
-					MaterialNode = CreateMaterialNode(NodeContainer, MaterialName, JSonErrorMessages);
+					MaterialNode = CreateMaterialNode(NodeContainer, NodeUid, MaterialName, JSonErrorMessages);
 					if (MaterialNode == nullptr)
 					{
 						JSonErrorMessages.Add(TEXT("{\"Msg\" : {\"Type\" : \"Error\",\n\"Msg\" : \"Cannot create fbx material (" + MaterialName + TEXT(")\"}}")));
@@ -75,7 +74,7 @@ namespace UE
 						{
 							int32 UnsupportedTextureCount = FbxProperty.GetSrcObjectCount<FbxLayeredTexture>();
 							UnsupportedTextureCount += FbxProperty.GetSrcObjectCount<FbxProceduralTexture>();
-							int32 TextureCount = FbxProperty.GetSrcObjectCount<FbxTexture>();
+							int32 TextureCount = FbxProperty.GetSrcObjectCount<FbxFileTexture>();
 							bool bFoundValidTexture = false;
 							if (UnsupportedTextureCount > 0)
 							{
@@ -94,11 +93,11 @@ namespace UE
 									}
 									//Create a texture node and make it child of the material node
 									TArray<FString> JsonErrorMessage;
-									FString NodeUID(TextureFilename);
-									UInterchangeTextureNode* TextureNode = Cast<UInterchangeTextureNode>(NodeContainer.GetNode(NodeUID));
+									FString NodeUid = TEXT("\\Texture\\") + TextureFilename;
+									UInterchangeTextureNode* TextureNode = Cast<UInterchangeTextureNode>(NodeContainer.GetNode(NodeUid));
 									if (!TextureNode)
 									{
-										TextureNode = CreateTexture2DNode(NodeContainer, TextureFilename, JsonErrorMessage);
+										TextureNode = CreateTexture2DNode(NodeContainer, NodeUid, TextureFilename, JsonErrorMessage);
 									}
 									// add/find UVSet and set it to the texture, we pass index 0 here, I think pipeline should be able to get the UVIndex channel from the name
 									// and modify the parameter to set the correct value.
@@ -108,7 +107,7 @@ namespace UE
 									float ScaleU = (float)FbxTextureFilePath->GetScaleU();
 									float ScaleV = (float)FbxTextureFilePath->GetScaleV();
 									MaterialNode->AddTextureParameterData(MaterialParameterName, TextureNode->GetUniqueID(), UVChannelIndex, ScaleU, ScaleV);
-									MaterialNode->SetDependencyUID(TextureNode->GetUniqueID());
+									MaterialNode->SetTextureDependencyUid(TextureNode->GetUniqueID());
 									bSetMaterial = true;
 									bFoundValidTexture = true;
 								}
@@ -172,7 +171,30 @@ namespace UE
 				return MaterialNode;
 			}
 
-			void FFbxMaterial::AddAllNodeMaterials(UInterchangeBaseNode* UnrealNode, FbxNode* ParentFbxNode, UInterchangeBaseNodeContainer& NodeContainer, TArray<FString>& JSonErrorMessages)
+			void FFbxMaterial::AddAllTextures(FbxScene* SDKScene, UInterchangeBaseNodeContainer& NodeContainer, TArray<FString>& JSonErrorMessages)
+			{
+				int32 TextureCount = SDKScene->GetSrcObjectCount<FbxFileTexture>();
+				for (int32 TextureIndex = 0; TextureIndex < TextureCount; ++TextureIndex)
+				{
+					FbxFileTexture* Texture = SDKScene->GetSrcObject<FbxFileTexture>(TextureIndex);
+					FString TextureFilename = UTF8_TO_TCHAR(Texture->GetFileName());
+					//Only import texture that exist on disk
+					if (!FPaths::FileExists(TextureFilename))
+					{
+						continue;
+					}
+					//Create a texture node and make it child of the material node
+					TArray<FString> JsonErrorMessage;
+					FString NodeUid = TEXT("\\Texture\\") + TextureFilename;
+					UInterchangeTextureNode* TextureNode = Cast<UInterchangeTextureNode>(NodeContainer.GetNode(NodeUid));
+					if (!TextureNode)
+					{
+						TextureNode = CreateTexture2DNode(NodeContainer, NodeUid, TextureFilename, JsonErrorMessage);
+					}
+				}
+			}
+			
+			void FFbxMaterial::AddAllNodeMaterials(UInterchangeSceneNode* SceneNode, FbxNode* ParentFbxNode, UInterchangeBaseNodeContainer& NodeContainer, TArray<FString>& JSonErrorMessages)
 			{
 				int32 MaterialCount = ParentFbxNode->GetMaterialCount();
 				for (int32 MaterialIndex = 0; MaterialIndex < MaterialCount; ++MaterialIndex)
@@ -180,16 +202,16 @@ namespace UE
 					FbxSurfaceMaterial* SurfaceMaterial = ParentFbxNode->GetMaterial(MaterialIndex);
 					UInterchangeMaterialNode* MaterialNode = AddNodeMaterial(SurfaceMaterial, NodeContainer, JSonErrorMessages);
 					//The dependencies order is important because mesh will use index in that order to determine material use by a face
-					UnrealNode->SetDependencyUID(MaterialNode->GetUniqueID());
+					SceneNode->AddAssetDependency(MaterialNode->GetUniqueID());
 				}
 			}
 
-			void FFbxMaterial::AddAllSceneMaterials(FbxScene* SDKScene, UInterchangeBaseNodeContainer& NodeContainer, TArray<FString>& JSonErrorMessages)
+			void FFbxMaterial::AddAllMaterials(FbxScene* SDKScene, UInterchangeBaseNodeContainer& NodeContainer, TArray<FString>& JSonErrorMessages)
 			{
-				int32 NodeMaterialCount = SDKScene->GetMaterialCount();
-				for (int32 NodeMaterialIndex = 0; NodeMaterialIndex < NodeMaterialCount; ++NodeMaterialIndex)
+				int32 MaterialCount = SDKScene->GetMaterialCount();
+				for (int32 MaterialIndex = 0; MaterialIndex < MaterialCount; ++MaterialIndex)
 				{
-					FbxSurfaceMaterial* SurfaceMaterial = SDKScene->GetMaterial(NodeMaterialIndex);
+					FbxSurfaceMaterial* SurfaceMaterial = SDKScene->GetMaterial(MaterialIndex);
 					AddNodeMaterial(SurfaceMaterial, NodeContainer, JSonErrorMessages);
 				}
 			}
diff --git a/Engine/Plugins/Runtime/Interchange/Source/Parsers/Fbx/Private/FbxMaterial.h b/Engine/Plugins/Runtime/Interchange/Source/Parsers/Fbx/Private/FbxMaterial.h
--- a/Engine/Plugins/Runtime/Interchange/Source/Parsers/Fbx/Private/FbxMaterial.h
+++ b/Engine/Plugins/Runtime/Interchange/Source/Parsers/Fbx/Private/FbxMaterial.h
@@ -10,6 +10,7 @@ class UInterchangeBaseNode;
 class UInterchangeBaseNodeContainer;
 class UInterchangeMaterialNode;
 class UInterchangeTextureNode;
+class UInterchangeSceneNode;
 
 namespace UE
 {
@@ -20,11 +21,31 @@ namespace UE
 			class FFbxMaterial
 			{
 			public:
-				static void AddAllNodeMaterials(UInterchangeBaseNode* ParentNode, FbxNode* ParentFbxNode, UInterchangeBaseNodeContainer& NodeContainer, TArray<FString>& JSonErrorMessages);
-				static void AddAllSceneMaterials(FbxScene* SDKScene, UInterchangeBaseNodeContainer& NodeContainer, TArray<FString>& JSonErrorMessages);
-				static UInterchangeMaterialNode* CreateMaterialNode(UInterchangeBaseNodeContainer& NodeContainer, const FString& NodeName, TArray<FString>& JSonErrorMessages);
-				static UInterchangeTextureNode* CreateTexture2DNode(UInterchangeBaseNodeContainer& NodeContainer, const FString& TextureFilePath, TArray<FString>& JSonErrorMessages);
-			
+				/**
+				 * Create a UInterchangeFextureNode and add it to the NodeContainer for all texture of type FbxFileTexture the fbx file contain.
+				 *
+				 * @note - Any node that already exist in the NodeContainer will not be created or modify.
+				 */
+				static void AddAllTextures(FbxScene* SDKScene, UInterchangeBaseNodeContainer& NodeContainer, TArray<FString>& JSonErrorMessages);
+				
+				/**
+				 * Create a UInterchangeMaterialNode and add it to the NodeContainer for all material of type FbxSurfaceMaterial the fbx file contain.
+				 * 
+				 * @note - Any node that already exist in the NodeContainer will not be created or modify.
+				 */
+				static void AddAllMaterials(FbxScene* SDKScene, UInterchangeBaseNodeContainer& NodeContainer, TArray<FString>& JSonErrorMessages);
+
+				/**
+				 * Create a UInterchangeMaterialNode and add it to the NodeContainer for all material of type FbxSurfaceMaterial the fbx ParentFbxNode contain.
+				 * It also set the dependencies of the node materials on the interchange ParentNode.
+				 * 
+				 * @note - Any material node that already exist in the NodeContainer will simply be add has a dependency.
+				 */
+				static void AddAllNodeMaterials(UInterchangeSceneNode* SceneNode, FbxNode* ParentFbxNode, UInterchangeBaseNodeContainer& NodeContainer, TArray<FString>& JSonErrorMessages);
+			protected:
+				static UInterchangeMaterialNode* CreateMaterialNode(UInterchangeBaseNodeContainer& NodeContainer, const FString& NodeUID, const FString& NodeName, TArray<FString>& JSonErrorMessages);
+				static UInterchangeTextureNode* CreateTexture2DNode(UInterchangeBaseNodeContainer& NodeContainer, const FString& NodeUID, const FString& TextureFilePath, TArray<FString>& JSonErrorMessages);
+
 			private:
 				static UInterchangeMaterialNode* AddNodeMaterial(FbxSurfaceMaterial* SurfaceMaterial, UInterchangeBaseNodeContainer& NodeContainer, TArray<FString>& JSonErrorMessages);
 			};
diff --git a/Engine/Plugins/Runtime/Interchange/Source/Parsers/Fbx/Private/FbxMesh.cpp b/Engine/Plugins/Runtime/Interchange/Source/Parsers/Fbx/Private/FbxMesh.cpp
--- a/Engine/Plugins/Runtime/Interchange/Source/Parsers/Fbx/Private/FbxMesh.cpp
+++ b/Engine/Plugins/Runtime/Interchange/Source/Parsers/Fbx/Private/FbxMesh.cpp
@@ -6,8 +6,11 @@
 #include "FbxConvert.h"
 #include "FbxHelper.h"
 #include "FbxInclude.h"
-#include "FbxSkeletalMesh.h"
+#include "InterchangeMeshNode.h"
 #include "MeshDescription.h"
+#include "Misc/FileHelper.h"
+#include "Nodes/InterchangeBaseNodeContainer.h"
+#include "Serialization/LargeMemoryWriter.h"
 #include "SkeletalMeshAttributes.h"
 #include "StaticMeshAttributes.h"
 
@@ -144,21 +147,6 @@ namespace UE
 					UniqueUVCount = FMath::Min<int32>(UniqueUVCount, MAX_MESH_TEXTURE_COORDS_MD);
 				}
 
-				int32 FindLightUVIndex() const
-				{
-					// See if any of our UV set entry names match LightMapUV.
-					for (int32 UVSetIdx = 0; UVSetIdx < UVSets.Num(); UVSetIdx++)
-					{
-						if (UVSets[UVSetIdx] == TEXT("LightMapUV"))
-						{
-							return UVSetIdx;
-						}
-					}
-
-					// not found
-					return INDEX_NONE;
-				}
-
 				// @param FaceCornerIndex usually TriangleIndex * 3 + CornerIndex but more complicated for mixed n-gons
 				int32 ComputeUVIndex(int32 UVLayerIndex, int32 lControlPointIndex, int32 FaceCornerIndex) const
 				{
@@ -197,31 +185,32 @@ namespace UE
 				int32 UniqueUVCount;
 			};
 
-			FMeshDescriptionImporter::FMeshDescriptionImporter(FMeshDescription* InMeshDescription, FbxNode* InMeshNode, FbxScene* InSDKScene, FbxGeometryConverter* InSDKGeometryConverter)
+			FMeshDescriptionImporter::FMeshDescriptionImporter(FMeshDescription* InMeshDescription, FbxScene* InSDKScene, FbxGeometryConverter* InSDKGeometryConverter)
 				: MeshDescription(InMeshDescription)
-				, MeshNode(InMeshNode)
 				, SDKScene(InSDKScene)
 				, SDKGeometryConverter(InSDKGeometryConverter)
 			{
-				bInitialized = ensure(MeshDescription) && ensure(MeshNode) && ensure(SDKScene) && ensure(SDKGeometryConverter);
+				bInitialized = ensure(MeshDescription) && ensure(SDKScene) && ensure(SDKGeometryConverter);
 			}
 			
-			bool FMeshDescriptionImporter::FillStaticMeshDescriptionFromFbxMesh()
+			bool FMeshDescriptionImporter::FillStaticMeshDescriptionFromFbxMesh(FbxMesh* Mesh)
 			{
-				return FillMeshDescriptionFromFbxMesh(EMeshType::Static, nullptr);
+				if (!ensure(bInitialized) || !ensure(MeshDescription) || !ensure(Mesh) || !ensure(Mesh->GetDeformerCount() == 0))
+				{
+					return false;
+				}
+				TArray<FString> UnusedArray;
+				return FillMeshDescriptionFromFbxMesh(Mesh, UnusedArray, EMeshType::Static);
 			}
 
-			bool FMeshDescriptionImporter::FillSkinnedMeshDescriptionFromFbxMesh(TArray<FbxNode*>* SortedJoints)
+			bool FMeshDescriptionImporter::FillSkinnedMeshDescriptionFromFbxMesh(FbxMesh* Mesh, TArray<FString>& OutJointNodeUniqueIDs)
 			{
-				FbxMesh* Mesh = MeshNode->GetMesh();
-				if (!ensure(Mesh))
+				if (!ensure(bInitialized) || !ensure(MeshDescription) || !ensure(Mesh) || !ensure(Mesh->GetDeformerCount() > 0))
 				{
-					//TODO log an error
 					return false;
 				}
-				FbxSkin* Skin = (FbxSkin*)Mesh->GetDeformer(0, FbxDeformer::eSkin);
-				EMeshType MeshType = Skin != nullptr ? EMeshType::Skinned : EMeshType::Rigid;
-				return FillMeshDescriptionFromFbxMesh(MeshType, SortedJoints);
+
+				return FillMeshDescriptionFromFbxMesh(Mesh, OutJointNodeUniqueIDs, EMeshType::Skinned);
 			}
 
 			bool FMeshDescriptionImporter::FillMeshDescriptionFromFbxShape(FbxShape* Shape)
@@ -232,10 +221,9 @@ namespace UE
 					return false;
 				}
 				FStaticMeshAttributes Attributes(*MeshDescription);
-				FbxMesh* Mesh = MeshNode->GetMesh();
 
 				//The shape should be for a specified mesh
-				if (!ensure(Mesh))
+				if (!ensure(Shape))
 				{
 					return false;
 				}
@@ -248,21 +236,13 @@ namespace UE
 					// Construct the matrices for the conversion from right handed to left handed system
 					FbxAMatrix TotalMatrix;
 					FbxAMatrix TotalMatrixForNormal;
-					TotalMatrix = ComputeNodeMatrix(MeshNode);
+					TotalMatrix.SetIdentity(); //We use the identity since we want to bake in the interchange factory if the pipeline request it
 					TotalMatrixForNormal = TotalMatrix.Inverse();
 					TotalMatrixForNormal = TotalMatrixForNormal.Transpose();
 
-					FbxGeometryBase* GeoBase = (Shape != nullptr ? (FbxGeometryBase*)Shape : (FbxGeometryBase*)Mesh);
+					FbxGeometryBase* GeoBase = (FbxGeometryBase*)Shape;
 
 					int32 VertexCount = GeoBase->GetControlPointsCount();
-					if (Shape)
-					{
-						if (!ensure(Mesh->GetControlPointsCount() == VertexCount))
-						{
-							//TODO warn the user
-							return false;
-						}
-					}
 
 					TVertexAttributesRef<FVector> VertexPositions = Attributes.GetVertexPositions();
 					int32 VertexOffset = MeshDescription->Vertices().Num();
@@ -299,14 +279,13 @@ namespace UE
 				return true;
 			}
 
-			bool FMeshDescriptionImporter::FillMeshDescriptionFromFbxMesh(EMeshType MeshType, TArray<FbxNode*>* SortedJoints)
+			bool FMeshDescriptionImporter::FillMeshDescriptionFromFbxMesh(FbxMesh* Mesh, TArray<FString>& OutJointNodeUniqueIDs, EMeshType MeshType)
 			{
-				if (!ensure(bInitialized) || !ensure(MeshDescription))
+				if (!ensure(bInitialized) || !ensure(MeshDescription) || !ensure(Mesh))
 				{
 					return false;
 				}
 
-				FbxMesh* Mesh = MeshNode->GetMesh();
 				FStaticMeshAttributes Attributes(*MeshDescription);
 
 				//Get the base layer of the mesh
@@ -314,32 +293,27 @@ namespace UE
 				if (BaseLayer == NULL)
 				{
 					//TODO add error message
-					//AddTokenizedErrorMessage(FTokenizedMessage::Create(EMessageSeverity::Error, FText::Format(LOCTEXT("Error_NoGeometryInMesh", "There is no geometry information in mesh '{0}'"), FText::FromString(Mesh->GetName()))), FFbxErrors::Generic_Mesh_NoGeometry);
 					return false;
 				}
 
 				FFBXUVs FBXUVs(Mesh);
 
-				if (MeshType == EMeshType::Static)
-				{
-// 					int32 FBXNamedLightMapCoordinateIndex = FBXUVs.FindLightUVIndex();
-// 					if (FBXNamedLightMapCoordinateIndex != INDEX_NONE)
-// 					{
-// 						StaticMesh->SetLightMapCoordinateIndex(FBXNamedLightMapCoordinateIndex);
-// 					}
-				}
-
+				FbxNode* MeshNode = Mesh->GetNodeCount() > 0 ? Mesh->GetNode() : nullptr; //Get the first node using the mesh to setup default asset materials
 				//
 				// create materials
 				//
 				//Create a material name array in the node order, also fill the Meshdescription PolygonGroup
 				TArray<FName> MaterialNames;
-				const int32 MaterialCount = MeshNode->GetMaterialCount();
-				MaterialNames.Reserve(MaterialCount);
-				for (int32 MaterialIndex = 0; MaterialIndex < MaterialCount; ++MaterialIndex)
+				int32 MaterialCount = (MeshNode != nullptr) ? MeshNode->GetMaterialCount() : Mesh->GetElementMaterialCount();
+				if (MeshNode)
 				{
-					FbxSurfaceMaterial* FbxMaterial = MeshNode->GetMaterial(MaterialIndex);
-					MaterialNames.Add(*FFbxHelper::GetFbxObjectName(FbxMaterial));
+					MaterialCount = MeshNode->GetMaterialCount();
+					MaterialNames.Reserve(MaterialCount);
+					for (int32 MaterialIndex = 0; MaterialIndex < MaterialCount; ++MaterialIndex)
+					{
+						FbxSurfaceMaterial* FbxMaterial = MeshNode->GetMaterial(MaterialIndex);
+						MaterialNames.Add(*FFbxHelper::GetFbxObjectName(FbxMaterial));
+					}
 				}
 
 				// Must do this before triangulating the mesh due to an FBX bug in TriangulateMeshAdvance
@@ -365,8 +339,8 @@ namespace UE
 					else
 					{
 						//TODO add an error message
-						//AddTokenizedErrorMessage(FTokenizedMessage::Create(EMessageSeverity::Warning, FText::Format(LOCTEXT("Error_FailedToTriangulate", "Unable to triangulate mesh '{0}'"), FText::FromString(Mesh->GetName()))), FFbxErrors::Generic_Mesh_TriangulationFailed);
-						return false; // not clean, missing some dealloc
+						MeshDescription->Empty();
+						return false;
 					}
 				}
 
@@ -402,7 +376,7 @@ namespace UE
 					//If all smooth group are 0's, we can say that there is no valid smoothing data, but only for skeletal mesh
 					//Static mesh will be faceted in this case since we will not compute smooth group from normals like we do here
 					//This is legacy code to get the same result has before
-					if (MeshType == EMeshType::Skinned || MeshType == EMeshType::Rigid)
+					if (MeshType == EMeshType::Skinned)
 					{
 						// Check and see if the smooothing data is valid.  If not generate it from the normals
 						BaseLayer = Mesh->GetLayer(0);
@@ -503,7 +477,7 @@ namespace UE
 					// Construct the matrices for the conversion from right handed to left handed system
 					FbxAMatrix TotalMatrix;
 					FbxAMatrix TotalMatrixForNormal;
-					TotalMatrix = ComputeNodeMatrix(MeshNode);
+					TotalMatrix.SetIdentity();
 					TotalMatrixForNormal = TotalMatrix.Inverse();
 					TotalMatrixForNormal = TotalMatrixForNormal.Transpose();
 					int32 PolygonCount = Mesh->GetPolygonCount();
@@ -825,17 +799,20 @@ namespace UE
 								}
 							}
 
-							if (MaterialIndex >= MaterialCount || MaterialIndex < 0)
+							if (MaterialIndex < 0)
 							{
 								//TODO log an error message
-								//AddTokenizedErrorMessage(FTokenizedMessage::Create(EMessageSeverity::Warning, LOCTEXT("Error_MaterialIndexInconsistency", "Face material index inconsistency - forcing to 0")), FFbxErrors::Generic_Mesh_MaterialIndexInconsistency);
 								MaterialIndex = 0;
 							}
-							if (!MaterialNames.IsValidIndex(MaterialIndex))
+							//Add missing material name with generated name
+							while (MaterialIndex >= MaterialNames.Num())
 							{
-								MaterialIndex = 0;
+								FString MaterialName = TEXT("Material_") + FString::FromInt(MaterialNames.Num());
+								MaterialNames.Add((*MaterialName));
 							}
-							else
+
+							//Material name should always exist since we create up to index materials
+							check(MaterialNames.IsValidIndex(MaterialIndex));
 							{
 								//Create a polygon with the 3 vertex instances Add it to the material group
 								FName MaterialName = MaterialNames[MaterialIndex];
@@ -1000,14 +977,8 @@ namespace UE
 						}
 					}
 
-					if (MeshType == EMeshType::Skinned || MeshType == EMeshType::Rigid)
+					if (MeshType == EMeshType::Skinned)
 					{
-						if (!ensure(SortedJoints))
-						{
-							//TODO log an error
-							return false;
-						}
-
 						FSkeletalMeshAttributes SkeletalMeshAttributes(*MeshDescription);
 						SkeletalMeshAttributes.Register();
 
@@ -1017,110 +988,89 @@ namespace UE
 						//Add the influence data in the skeletalmesh description
 						FSkinWeightsVertexAttributesRef VertexSkinWeights = SkeletalMeshAttributes.GetVertexSkinWeights();
 
-						if (MeshType == EMeshType::Skinned)
+						FbxSkin* Skin = (FbxSkin*)Mesh->GetDeformer(0, FbxDeformer::eSkin);
+						if (!ensure(Skin))
+						{
+							//TODO log an error
+							return false;
+						}
+						const int32 ClusterCount = Skin->GetClusterCount();
+						TArray<FbxNode*> SortedJoints;
+						SortedJoints.Reserve(ClusterCount);
+						OutJointNodeUniqueIDs.Reserve(ClusterCount);
+						// create influences for each cluster
+						for (int32 ClusterIndex = 0; ClusterIndex < ClusterCount; ClusterIndex++)
 						{
-							FbxSkin* Skin = (FbxSkin*)Mesh->GetDeformer(0, FbxDeformer::eSkin);
-							if (!ensure(Skin))
+							FbxCluster* Cluster = Skin->GetCluster(ClusterIndex);
+							// When Maya plug-in exports rigid binding, it will generate "CompensationCluster" for each ancestor links.
+							// FBX writes these "CompensationCluster" out. The CompensationCluster also has weight 1 for vertices.
+							// Unreal importer should skip these clusters.
+							if (!Cluster || (FCStringAnsi::Strcmp(Cluster->GetUserDataID(), "Maya_ClusterHint") == 0 && FCStringAnsi::Strcmp(Cluster->GetUserData(), "CompensationCluster") == 0))
 							{
-								//TODO log an error
-								return false;
+								continue;
 							}
-							// create influences for each cluster
-							for (int32 ClusterIndex = 0; ClusterIndex < Skin->GetClusterCount(); ClusterIndex++)
+
+							FbxNode* Link = Cluster->GetLink();
+							// find the bone index
+							int32 BoneIndex = -1;
+							if (!SortedJoints.Find(Link, BoneIndex))
+							{
+								BoneIndex = SortedJoints.Add(Link);
+								FString JointNodeUniqueID = FFbxHelper::GetFbxNodeHierarchyName(Link);
+								OutJointNodeUniqueIDs.Add(JointNodeUniqueID);
+							}
+
+							//	get the vertex indices
+							int32 ControlPointIndicesCount = Cluster->GetControlPointIndicesCount();
+							int32* ControlPointIndices = Cluster->GetControlPointIndices();
+							double* Weights = Cluster->GetControlPointWeights();
+
+							//	for each vertex index in the cluster
+							for (int32 ControlPointIndex = 0; ControlPointIndex < ControlPointIndicesCount; ++ControlPointIndex)
 							{
-								FbxCluster* Cluster = Skin->GetCluster(ClusterIndex);
-								// When Maya plug-in exports rigid binding, it will generate "CompensationCluster" for each ancestor links.
-								// FBX writes these "CompensationCluster" out. The CompensationCluster also has weight 1 for vertices.
-								// Unreal importer should skip these clusters.
-								if (!Cluster || (FCStringAnsi::Strcmp(Cluster->GetUserDataID(), "Maya_ClusterHint") == 0 && FCStringAnsi::Strcmp(Cluster->GetUserData(), "CompensationCluster") == 0))
+								FVertexID VertexID(ControlPointIndices[ControlPointIndex] + VertexOffset);
+								if (!ensure(MeshDescription->IsVertexValid(VertexID)))
 								{
+									//Invalid Influence
 									continue;
 								}
+								float Weight = static_cast<float>(Weights[ControlPointIndex]);
 
-								FbxNode* Link = Cluster->GetLink();
-								// find the bone index
-								int32 BoneIndex = -1;
-								SortedJoints->Find(Link, BoneIndex);
-
-								//	get the vertex indices
-								int32 ControlPointIndicesCount = Cluster->GetControlPointIndicesCount();
-								int32* ControlPointIndices = Cluster->GetControlPointIndices();
-								double* Weights = Cluster->GetControlPointWeights();
-
-								//	for each vertex index in the cluster
-								for (int32 ControlPointIndex = 0; ControlPointIndex < ControlPointIndicesCount; ++ControlPointIndex)
+								TArray<FBoneWeight> *BoneWeights = RawBoneWeights.Find(VertexID);
+								if (BoneWeights)
 								{
-									FVertexID VertexID(ControlPointIndices[ControlPointIndex] + VertexOffset);
-									if (!ensure(MeshDescription->IsVertexValid(VertexID)))
-									{
-										//Invalid Influence
-										continue;
-									}
-									float Weight = static_cast<float>(Weights[ControlPointIndex]);
-
-									TArray<FBoneWeight> *BoneWeights = RawBoneWeights.Find(VertexID);
-									if (BoneWeights)
+									// Do we already have a weight for this bone? Keep the greater weight.
+									bool bShouldAdd = true;
+									for (int32 WeightIndex = 0; WeightIndex < BoneWeights->Num(); WeightIndex++)
 									{
-										// Do we already have a weight for this bone? Keep the greater weight.
-										bool bShouldAdd = true;
-										for (int32 WeightIndex = 0; WeightIndex < BoneWeights->Num(); WeightIndex++)
+										FBoneWeight &BoneWeight = (*BoneWeights)[WeightIndex];
+										if (BoneWeight.GetBoneIndex() == BoneIndex)
 										{
-											FBoneWeight &BoneWeight = (*BoneWeights)[WeightIndex];
-											if (BoneWeight.GetBoneIndex() == BoneIndex)
+											if (BoneWeight.GetWeight() < Weight)
 											{
-												if (BoneWeight.GetWeight() < Weight)
-												{
-													BoneWeight.SetWeight(Weight);
-												}
-												bShouldAdd = false;
-												break;
+												BoneWeight.SetWeight(Weight);
 											}
-										}
-										if (bShouldAdd)
-										{
-											BoneWeights->Add(FBoneWeight(BoneIndex, Weight));
+											bShouldAdd = false;
+											break;
 										}
 									}
-									else
+									if (bShouldAdd)
 									{
-										RawBoneWeights.Add(VertexID).Add(FBoneWeight(BoneIndex, Weight));
+										BoneWeights->Add(FBoneWeight(BoneIndex, Weight));
 									}
 								}
-							}
-
-							// Add all the raw bone weights. This will cause the weights to be sorted and re-normalized after culling to max influences.
-							for (const TTuple<FVertexID, TArray<FBoneWeight>> &Item: RawBoneWeights)
-							{
-								VertexSkinWeights.Set(Item.Key, Item.Value);
-							}	
-						}
-						else // for rigid mesh
-						{
-							// find the bone index, the bone is the node itself
-							int32 BoneIndex = -1;
-							SortedJoints->Find(MeshNode, BoneIndex);
-
-							const float Weight = 1.0f;
-							const FBoneWeights DefaultBinding = FBoneWeights::Create({FBoneWeight(BoneIndex, Weight)});
-							if (BoneIndex == -1 && SortedJoints->Num() > 0)
-							{
-								//When we import geometry only, we want to hook all the influence to a particular bone
-								BoneIndex = 0;
-							}
-							//	for each vertex in the mesh
-							for (int32 ControlPointIndex = 0; ControlPointIndex < VertexCount; ++ControlPointIndex)
-							{
-								FVertexID VertexID(ControlPointIndex + VertexOffset);
-								if (!ensure(MeshDescription->IsVertexValid(VertexID)))
+								else
 								{
-									//Invalid Influence
-									continue;
+									RawBoneWeights.Add(VertexID).Add(FBoneWeight(BoneIndex, Weight));
 								}
-
-								VertexSkinWeights.Set(VertexID, DefaultBinding);
 							}
 						}
 
+						// Add all the raw bone weights. This will cause the weights to be sorted and re-normalized after culling to max influences.
+						for (const TTuple<FVertexID, TArray<FBoneWeight>> &Item: RawBoneWeights)
+						{
+							VertexSkinWeights.Set(Item.Key, Item.Value);
+						}
 					}
 				}
 				// needed?
@@ -1135,23 +1085,6 @@ namespace UE
 				return bHasNonDegeneratePolygons;
 			}
 
-			FbxAMatrix FMeshDescriptionImporter::ComputeNodeMatrix(FbxNode* Node)
-			{
-				FbxAMatrix Geometry;
-				FbxVector4 Translation, Rotation, Scaling;
-				Translation = Node->GetGeometricTranslation(FbxNode::eSourcePivot);
-				Rotation = Node->GetGeometricRotation(FbxNode::eSourcePivot);
-				Scaling = Node->GetGeometricScaling(FbxNode::eSourcePivot);
-				Geometry.SetT(Translation);
-				Geometry.SetR(Rotation);
-				Geometry.SetS(Scaling);
-				//For Single Matrix situation, obtain transfrom matrix from eDESTINATION_SET, which include pivot offsets and pre/post rotations.
-				FbxAMatrix& GlobalTransform = SDKScene->GetAnimationEvaluator()->GetNodeGlobalTransform(Node);
-				//We must always add the geometric transform. Only Max use the geometric transform which is an offset to the local transform of the node
-				FbxAMatrix NodeMatrix = GlobalTransform * Geometry;
-				return NodeMatrix;
-			}
-
 			bool FMeshDescriptionImporter::IsOddNegativeScale(FbxAMatrix& TotalMatrix)
 			{
 				FbxVector4 Scale = TotalMatrix.GetS();
@@ -1163,6 +1096,375 @@ namespace UE
 
 				return NegativeNum == 1 || NegativeNum == 3;
 			}
+
+
+			//////////////////////////////////////////////////////////////////////////
+			/// FMeshPayloadContext implementation
+
+			bool FMeshPayloadContext::FetchPayloadToFile(const FString& PayloadFilepath, TArray<FString>& JSonErrorMessages)
+			{
+				if (!ensure(SDKScene != nullptr))
+				{
+					JSonErrorMessages.Add(TEXT("{\"Msg\" : {\"Type\" : \"Error\",\n\"Msg\" : \"Cannot fetch fbx mesh payload because the fbx scene is null\"}}"));
+					return false;
+				}
+
+				if (!ensure(Mesh != nullptr))
+				{
+					JSonErrorMessages.Add(TEXT("{\"Msg\" : {\"Type\" : \"Error\",\n\"Msg\" : \"Cannot fetch fbx mesh payload because the fbx mesh is null\"}}"));
+					return false;
+				}
+
+				if (!ensure(SDKGeometryConverter != nullptr))
+				{
+					JSonErrorMessages.Add(TEXT("{\"Msg\" : {\"Type\" : \"Error\",\n\"Msg\" : \"Cannot fetch fbx mesh payload because the fbx scene geometry converter is null\"}}"));
+					return false;
+				}
+
+				bool bFetchSkinnedData = (Mesh->GetDeformerCount(FbxDeformer::eSkin) > 0);
+				
+				TArray<FString> JointUniqueID;
+
+				FMeshDescription MeshDescription;
+				if (bFetchSkinnedData)
+				{
+					FSkeletalMeshAttributes SkeletalMeshAttribute(MeshDescription);
+					SkeletalMeshAttribute.Register();
+					FMeshDescriptionImporter MeshDescriptionImporter(&MeshDescription, SDKScene, SDKGeometryConverter);
+					if (!MeshDescriptionImporter.FillSkinnedMeshDescriptionFromFbxMesh(Mesh, JointUniqueID))
+					{
+						JSonErrorMessages.Add(TEXT("{\"Msg\" : {\"Type\" : \"Error\",\n\"Msg\" : \"Cannot fetch skinned mesh payload because there was an error when creating the FMeshDescription\"}}"));
+						return false;
+					}
+				}
+				else
+				{
+					FStaticMeshAttributes StaticMeshAttribute(MeshDescription);
+					StaticMeshAttribute.Register();
+					FMeshDescriptionImporter MeshDescriptionImporter(&MeshDescription, SDKScene, SDKGeometryConverter);
+					if (!MeshDescriptionImporter.FillStaticMeshDescriptionFromFbxMesh(Mesh))
+					{
+						JSonErrorMessages.Add(TEXT("{\"Msg\" : {\"Type\" : \"Error\",\n\"Msg\" : \"Cannot fetch static mesh payload because there was an error when creating the FMeshDescription\"}}"));
+						return false;
+					}
+				}
+				//Dump the MeshDescription to a file
+				{
+					FLargeMemoryWriter Ar;
+					MeshDescription.Serialize(Ar);
+
+					Ar << bFetchSkinnedData;
+					if (bFetchSkinnedData)
+					{
+						//When passing a skinned MeshDescription, We want to pass the joint Node ID so we can know what the influence bone index refer to
+						Ar << JointUniqueID;
+					}
+					uint8* ArchiveData = Ar.GetData();
+					int64 ArchiveSize = Ar.TotalSize();
+					TArray64<uint8> Buffer(ArchiveData, ArchiveSize);
+					FFileHelper::SaveArrayToFile(Buffer, *PayloadFilepath);
+				}
+
+				return true;
+			}
+
+			//////////////////////////////////////////////////////////////////////////
+			/// FShapePayloadContext implementation
+
+			bool FShapePayloadContext::FetchPayloadToFile(const FString& PayloadFilepath, TArray<FString>& JSonErrorMessages)
+			{
+				if (!ensure(Shape))
+				{
+					//Todo log an error
+					return false;
+				}
+				//Import the BlendShape
+				FMeshDescription ShapeMeshDescription;
+				FStaticMeshAttributes StaticMeshAttribute(ShapeMeshDescription);
+				StaticMeshAttribute.Register();
+				FMeshDescriptionImporter MeshDescriptionImporter(&ShapeMeshDescription, SDKScene, SDKGeometryConverter);
+				if (!MeshDescriptionImporter.FillMeshDescriptionFromFbxShape(Shape))
+				{
+					//Todo log an error
+					return false;
+				}
+
+				//Dump the MeshDescription to a file
+				{
+					FLargeMemoryWriter Ar;
+					ShapeMeshDescription.Serialize(Ar);
+					uint8* ArchiveData = Ar.GetData();
+					int64 ArchiveSize = Ar.TotalSize();
+					TArray64<uint8> Buffer(ArchiveData, ArchiveSize);
+					FFileHelper::SaveArrayToFile(Buffer, *PayloadFilepath);
+				}
+				return true;
+			}
+
+			//////////////////////////////////////////////////////////////////////////
+			/// FFbxMesh implementation
+			
+			FString FFbxMesh::GetUniqueIDString(const uint64 UniqueID)
+			{
+				FStringFormatNamedArguments FormatArguments;
+				FormatArguments.Add(TEXT("UniqueID"), UniqueID);
+				return FString::Format(TEXT("{UniqueID}"), FormatArguments);
+			}
+
+			FString FFbxMesh::GetMeshName(FbxGeometryBase* Mesh)
+			{
+				FString MeshName = FFbxHelper::GetFbxObjectName(Mesh);
+				if (MeshName.IsEmpty())
+				{
+					if (Mesh->GetNodeCount() > 0)
+					{
+						if (Mesh->GetAttributeType() == FbxNodeAttribute::eMesh)
+						{
+							MeshName = TEXT("Mesh_");
+						}
+						else if (Mesh->GetAttributeType() == FbxNodeAttribute::eShape)
+						{
+							MeshName = TEXT("Shape_");
+						}
+						MeshName += FFbxHelper::GetFbxObjectName(Mesh->GetNode(0));
+					}
+					else
+					{
+						uint64 UniqueFbxObjectID = Mesh->GetUniqueID();
+						MeshName += GetUniqueIDString(UniqueFbxObjectID);
+					}
+				}
+				return MeshName;
+			}
+			
+			FString FFbxMesh::GetMeshUniqueID(FbxGeometryBase* Mesh)
+			{
+				FString MeshUniqueID;
+				if (Mesh->GetAttributeType() == FbxNodeAttribute::eMesh)
+				{
+					MeshUniqueID = TEXT("\\Mesh\\");
+				}
+				else if (Mesh->GetAttributeType() == FbxNodeAttribute::eShape)
+				{
+					MeshUniqueID = TEXT("\\Shape\\");
+				}
+				FString MeshName = FFbxHelper::GetFbxObjectName(Mesh);
+				if (MeshName.IsEmpty())
+				{
+					if (Mesh->GetNodeCount() > 0)
+					{
+						MeshUniqueID += FFbxHelper::GetFbxNodeHierarchyName(Mesh->GetNode(0));
+					}
+					else
+					{
+						MeshUniqueID += GetMeshName(Mesh);
+					}
+				}
+				else
+				{
+					MeshUniqueID += MeshName;
+				}
+				return MeshUniqueID;
+			}
+
+			void FFbxMesh::ExtractSkinnedMeshNodeJoints(FbxScene* SDKScene, FbxMesh* Mesh, UInterchangeMeshNode* MeshNode, TArray<FString>& JSonErrorMessages)
+			{
+				TArray<FString> JointNodeUniqueIDs;
+				const int32 SkinDeformerCount = Mesh->GetDeformerCount(FbxDeformer::eSkin);
+				for (int32 DeformerIndex = 0; DeformerIndex < SkinDeformerCount; DeformerIndex++)
+				{
+					FbxSkin* Skin = (FbxSkin*)Mesh->GetDeformer(DeformerIndex, FbxDeformer::eSkin);
+					if (!ensure(Skin))
+					{
+						continue;
+					}
+					const int32 ClusterCount = Skin->GetClusterCount();
+					JointNodeUniqueIDs.Reserve(ClusterCount);
+					for (int32 ClusterIndex = 0; ClusterIndex < ClusterCount; ClusterIndex++)
+					{
+						FbxCluster* Cluster = Skin->GetCluster(ClusterIndex);
+						// When Maya plug-in exports rigid binding, it will generate "CompensationCluster" for each ancestor links.
+						// FBX writes these "CompensationCluster" out. The CompensationCluster also has weight 1 for vertices.
+						// Unreal importer should skip these clusters.
+						if (!Cluster || (FCStringAnsi::Strcmp(Cluster->GetUserDataID(), "Maya_ClusterHint") == 0 && FCStringAnsi::Strcmp(Cluster->GetUserData(), "CompensationCluster") == 0))
+						{
+							continue;
+						}
+						FbxNode* Link = Cluster->GetLink();
+						FString JointNodeUniqueID = FFbxHelper::GetFbxNodeHierarchyName(Link);
+						// find the bone index
+						if (!JointNodeUniqueIDs.Contains(JointNodeUniqueID))
+						{
+							JointNodeUniqueIDs.Add(JointNodeUniqueID);
+							MeshNode->SetSkeletonDependencyUid(JointNodeUniqueID);
+						}
+					}
+				}
+			}
+
+			void FFbxMesh::AddAllMeshes(FbxScene* SDKScene, FbxGeometryConverter* SDKGeometryConverter, UInterchangeBaseNodeContainer& NodeContainer, TArray<FString>& JSonErrorMessages, TMap<FString, TSharedPtr<FPayloadContextBase>>& PayloadContexts)
+			{
+				int32 GeometryCount = SDKScene->GetGeometryCount();
+				for (int32 GeometryIndex = 0; GeometryIndex < GeometryCount; ++GeometryIndex)
+				{
+					FbxGeometry* Geometry = SDKScene->GetGeometry(GeometryIndex);
+					if (Geometry->GetAttributeType() != FbxNodeAttribute::eMesh)
+					{
+						continue;
+					}
+					FbxMesh* Mesh = static_cast<FbxMesh*>(Geometry);
+					if (!Mesh)
+					{
+						continue;
+					}
+					FString MeshName = GetMeshName(Mesh);
+					FString MeshUniqueID = GetMeshUniqueID(Mesh);
+					UInterchangeMeshNode* MeshNode = Cast<UInterchangeMeshNode>(NodeContainer.GetNode(MeshUniqueID));
+					if (!MeshNode)
+					{
+						MeshNode = CreateMeshNode(NodeContainer, MeshName, MeshUniqueID, JSonErrorMessages);
+						if (Geometry->GetDeformerCount(FbxDeformer::eSkin) > 0)
+						{
+							//Set the skinned mesh attribute
+							MeshNode->SetSkinnedMesh(true);
+							ExtractSkinnedMeshNodeJoints(SDKScene, Mesh, MeshNode, JSonErrorMessages);
+						}
+						Mesh->ComputeBBox();
+						const int32 MeshVertexCount = Mesh->GetControlPointsCount();
+						MeshNode->SetCustomVertexCount(MeshVertexCount);
+						const int32 MeshPolygonCount = Mesh->GetPolygonCount();
+						MeshNode->SetCustomPolygonCount(MeshPolygonCount);
+						const FBox MeshBoundingBox = FBox(FFbxConvert::ConvertPos(Mesh->BBoxMin.Get()), FFbxConvert::ConvertPos(Mesh->BBoxMax.Get()));
+						MeshNode->SetCustomBoundingBox(MeshBoundingBox);
+						const bool bMeshHasVertexNormal = Mesh->GetElementNormalCount() > 0;
+						MeshNode->SetCustomHasVertexNormal(bMeshHasVertexNormal);
+						const bool bMeshHasVertexBinormal = Mesh->GetElementBinormalCount() > 0;
+						MeshNode->SetCustomHasVertexBinormal(bMeshHasVertexBinormal);
+						const bool bMeshHasVertexTangent = Mesh->GetElementTangentCount() > 0;
+						MeshNode->SetCustomHasVertexTangent(bMeshHasVertexTangent);
+						const bool bMeshHasSmoothGroup = Mesh->GetElementSmoothingCount() > 0;
+						MeshNode->SetCustomHasSmoothGroup(bMeshHasSmoothGroup);
+						const bool bMeshHasVertexColor = Mesh->GetElementVertexColorCount() > 0;
+						MeshNode->SetCustomHasVertexColor(bMeshHasVertexColor);
+						const int32 MeshUVCount = Mesh->GetElementUVCount();
+						MeshNode->SetCustomUVCount(MeshUVCount);
+						
+						//Add Material dependencies, we use always the first fbx node that instance the fbx geometry to grab the materials.
+						if (FbxNode* FbxMeshNode = Mesh->GetNode())
+						{
+							const int32 MaterialCount = FbxMeshNode->GetMaterialCount();
+							for (int32 MaterialIndex = 0; MaterialIndex < MaterialCount; ++MaterialIndex)
+							{
+								FbxSurfaceMaterial* FbxMaterial = FbxMeshNode->GetMaterial(MaterialIndex);
+								FString MaterialName = FFbxHelper::GetFbxObjectName(FbxMaterial);
+								FString MaterialUid = TEXT("\\Material\\") + MaterialName;
+								MeshNode->SetMaterialDependencyUid(MaterialUid);
+							}
+						}
+					}
+					FString PayLoadKey = MeshUniqueID;
+					if (ensure(!PayloadContexts.Contains(PayLoadKey)))
+					{
+						TSharedPtr<FMeshPayloadContext> GeoPayload = MakeShared<FMeshPayloadContext>();
+						GeoPayload->Mesh = Mesh;
+						GeoPayload->SDKScene = SDKScene;
+						GeoPayload->SDKGeometryConverter = SDKGeometryConverter;
+						PayloadContexts.Add(PayLoadKey, GeoPayload);
+					}
+					MeshNode->SetPayLoadKey(PayLoadKey);
+					
+					//Add all BlendShapes for this geometry node
+					TMap<FString, FbxShape*> ShapeNameToFbxShape;
+					const int32 BlendShapeCount = Geometry->GetDeformerCount(FbxDeformer::eBlendShape);
+					if (BlendShapeCount > 0)
+					{
+						for (int32 BlendShapeIndex = 0; BlendShapeIndex < BlendShapeCount; ++BlendShapeIndex)
+						{
+							FbxBlendShape* BlendShape = (FbxBlendShape*)Geometry->GetDeformer(BlendShapeIndex, FbxDeformer::eBlendShape);
+							const int32 BlendShapeChannelCount = BlendShape->GetBlendShapeChannelCount();
+							FString BlendShapeName = FFbxHelper::GetFbxObjectName(BlendShape);
+							// see below where this is used for explanation...
+							const bool bMightBeBadMAXFile = (BlendShapeName == FString("Morpher"));
+							for (int32 ChannelIndex = 0; ChannelIndex < BlendShapeChannelCount; ++ChannelIndex)
+							{
+								FbxBlendShapeChannel* Channel = BlendShape->GetBlendShapeChannel(ChannelIndex);
+								if (Channel)
+								{
+									//Find which shape should we use according to the weight.
+									const int32 CurrentChannelShapeCount = Channel->GetTargetShapeCount();
+									FString ChannelName = FFbxHelper::GetFbxObjectName(Channel);
+									// Maya adds the name of the blendshape and an underscore to the front of the channel name, so remove it
+									if (ChannelName.StartsWith(BlendShapeName))
+									{
+										ChannelName.RightInline(ChannelName.Len() - (BlendShapeName.Len() + 1), false);
+									}
+									for (int32 ShapeIndex = 0; ShapeIndex < CurrentChannelShapeCount; ++ShapeIndex)
+									{
+										FbxShape* Shape = Channel->GetTargetShape(ShapeIndex);
+										FString ShapeName;
+										if (CurrentChannelShapeCount > 1)
+										{
+											ShapeName = FFbxHelper::GetFbxObjectName(Shape);
+										}
+										else
+										{
+											if (bMightBeBadMAXFile)
+											{
+												ShapeName = FFbxHelper::GetFbxObjectName(Shape);
+											}
+											else
+											{
+												// Maya concatenates the number of the shape to the end of its name, so instead use the name of the channel
+												ShapeName = ChannelName;
+											}
+										}
+										ensure(!ShapeNameToFbxShape.Contains(ShapeName));
+										ShapeNameToFbxShape.Add(ShapeName, Shape);
+										FString ShapeAttributeName = GetMeshName(Shape);
+										FString ShapeUniqueID = GetMeshUniqueID(Shape);
+										UInterchangeMeshNode* ShapeNode = Cast<UInterchangeMeshNode>(NodeContainer.GetNode(ShapeUniqueID));
+										if (!ShapeNode)
+										{
+											ShapeNode = CreateMeshNode(NodeContainer, ShapeAttributeName, ShapeUniqueID, JSonErrorMessages);
+											const bool bIsBlendShape = true;
+											ShapeNode->SetBlendShape(bIsBlendShape);
+											ShapeNode->SetBlendShapeName(ShapeName);
+											//Create a Mesh node dependency, so the mesh node can retrieve is associate blendshape
+											MeshNode->SetShapeDependencyUid(ShapeUniqueID);
+										}
+										FString ShapePayLoadKey = ShapeUniqueID;
+										if (ensure(!PayloadContexts.Contains(ShapePayLoadKey)))
+										{
+											TSharedPtr<FShapePayloadContext> GeoPayload = MakeShared<FShapePayloadContext>();
+											GeoPayload->Shape = Shape;
+											GeoPayload->SDKScene = SDKScene;
+											GeoPayload->SDKGeometryConverter = SDKGeometryConverter;
+											PayloadContexts.Add(ShapePayLoadKey, GeoPayload);
+										}
+										ShapeNode->SetPayLoadKey(ShapePayLoadKey);
+									}
+								}
+							}
+						}
+					}
+				}
+			}
+
+			UInterchangeMeshNode* FFbxMesh::CreateMeshNode(UInterchangeBaseNodeContainer& NodeContainer, const FString& NodeName, const FString& NodeUniqueID, TArray<FString>& JSonErrorMessages)
+			{
+				FString DisplayLabel(NodeName);
+				FString NodeUid(NodeUniqueID);
+				UInterchangeMeshNode* MeshNode = NewObject<UInterchangeMeshNode>(&NodeContainer, NAME_None);
+				if (!ensure(MeshNode))
+				{
+					JSonErrorMessages.Add(TEXT("{\"Msg\" : {\"Type\" : \"Error\",\n\"Msg\" : \"Cannot allocate a mesh node when importing fbx\"}}"));
+					return nullptr;
+				}
+				// Creating a UMaterialInterface
+				MeshNode->InitializeNode(NodeUid, DisplayLabel, EInterchangeNodeContainerType::NodeContainerType_TranslatedAsset);
+				NodeContainer.AddNode(MeshNode);
+				return MeshNode;
+			}
 		} //ns Private
 	} //ns Interchange
 }//ns UE
diff --git a/Engine/Plugins/Runtime/Interchange/Source/Parsers/Fbx/Private/FbxMesh.h b/Engine/Plugins/Runtime/Interchange/Source/Parsers/Fbx/Private/FbxMesh.h
--- a/Engine/Plugins/Runtime/Interchange/Source/Parsers/Fbx/Private/FbxMesh.h
+++ b/Engine/Plugins/Runtime/Interchange/Source/Parsers/Fbx/Private/FbxMesh.h
@@ -8,7 +8,8 @@
 
 /** Forward declarations */
 struct FMeshDescription;
-
+class UInterchangeBaseNodeContainer;
+class UInterchangeMeshNode;
 
 namespace UE
 {
@@ -19,10 +20,23 @@ namespace UE
 			class FMeshDescriptionImporter
 			{
 			public:
-				FMeshDescriptionImporter(FMeshDescription* InMeshDescription, FbxNode* InMeshNode, FbxScene* InSDKScene, FbxGeometryConverter* InSDKGeometryConverter);
-				bool FillStaticMeshDescriptionFromFbxMesh();
-				bool FillSkinnedMeshDescriptionFromFbxMesh(TArray<FbxNode*>* SortedJoints);
+				FMeshDescriptionImporter(FMeshDescription* InMeshDescription, FbxScene* InSDKScene, FbxGeometryConverter* InSDKGeometryConverter);
+				
+				/*
+				 * Fill the mesh description using the Mesh parameter.
+				 */
+				bool FillStaticMeshDescriptionFromFbxMesh(FbxMesh* Mesh);
+				
+				/*
+				 * Fill the mesh description using the Mesh parameter and also fill the OutJointNodeUniqueIDs so the MeshDescription bone Index can be map to the correct interchange joint scene node.
+				 */
+				bool FillSkinnedMeshDescriptionFromFbxMesh(FbxMesh* Mesh, TArray<FString>& OutJointNodeUniqueIDs);
+
+				/*
+				 * Fill the mesh description using the Shape parameter.
+				 */
 				bool FillMeshDescriptionFromFbxShape(FbxShape* Shape);
+
 			private:
 				
 				enum class EMeshType : uint8
@@ -30,11 +44,9 @@ namespace UE
 					None = 0, //No mesh type to import
 					Static = 1, //static mesh
 					Skinned = 2, //skinned mesh with joints
-					Rigid = 3  //rigid mesh (joints will be created from the geometry transform)
 				};
 
-				bool FillMeshDescriptionFromFbxMesh(EMeshType MeshType, TArray<FbxNode*>* SortedJoints);
-				FbxAMatrix ComputeNodeMatrix(FbxNode* Node);
+				bool FillMeshDescriptionFromFbxMesh(FbxMesh* Mesh, TArray<FString>& OutJointNodeUniqueIDs, EMeshType MeshType);
 				bool IsOddNegativeScale(FbxAMatrix& TotalMatrix);
 				
 				//TODO move the real function from RenderCore to FVector, so we do not have to add render core to compute such a simple thing
@@ -48,15 +60,47 @@ namespace UE
 					);
 					return (Basis.Determinant() < 0) ? -1.0f : +1.0f;
 				}
-
-
 				FMeshDescription* MeshDescription;
-				FbxNode* MeshNode;
 				FbxScene* SDKScene;
 				FbxGeometryConverter* SDKGeometryConverter;
-				bool bIsStaticMesh;
 				bool bInitialized = false;
 			};
+
+			class FMeshPayloadContext : public FPayloadContextBase
+			{
+			public:
+				virtual ~FMeshPayloadContext() {}
+				virtual FString GetPayloadType() const override { return TEXT("Mesh-PayloadContext"); }
+				virtual bool FetchPayloadToFile(const FString& PayloadFilepath, TArray<FString>& JSonErrorMessages) override;
+				FbxMesh* Mesh = nullptr;
+				FbxScene* SDKScene = nullptr;
+				FbxGeometryConverter* SDKGeometryConverter = nullptr;
+			};
+
+			class FShapePayloadContext : public FPayloadContextBase
+			{
+			public:
+				virtual ~FShapePayloadContext() {}
+				virtual FString GetPayloadType() const override { return TEXT("Shape-PayloadContext"); }
+				virtual bool FetchPayloadToFile(const FString& PayloadFilepath, TArray<FString>& JSonErrorMessages) override;
+				FbxShape* Shape = nullptr;
+				FbxScene* SDKScene = nullptr;
+				FbxGeometryConverter* SDKGeometryConverter = nullptr;
+			};
+
+			class FFbxMesh
+			{
+			public:
+				static FString GetMeshName(FbxGeometryBase* Mesh);
+				static FString GetMeshUniqueID(FbxGeometryBase* Mesh);
+				static void ExtractSkinnedMeshNodeJoints(FbxScene* SDKScene, FbxMesh* Mesh, UInterchangeMeshNode* MeshNode, TArray<FString>& JSonErrorMessages);
+
+				static void AddAllMeshes(FbxScene* SDKScene, FbxGeometryConverter* SDKGeometryConverter, UInterchangeBaseNodeContainer& NodeContainer, TArray<FString>& JSonErrorMessages, TMap<FString, TSharedPtr<FPayloadContextBase>>& PayloadContexts);
+			
+			protected:
+				static UInterchangeMeshNode* CreateMeshNode(UInterchangeBaseNodeContainer& NodeContainer, const FString& NodeName, const FString& NodeUniqueID, TArray<FString>& JSonErrorMessages);
+				static FString GetUniqueIDString(const uint64 UniqueID);
+			};
 		}//ns Private
 	}//ns Interchange
 }//ns UE
diff --git a/Engine/Plugins/Runtime/Interchange/Source/Parsers/Fbx/Private/FbxScene.cpp b/Engine/Plugins/Runtime/Interchange/Source/Parsers/Fbx/Private/FbxScene.cpp
--- a/Engine/Plugins/Runtime/Interchange/Source/Parsers/Fbx/Private/FbxScene.cpp
+++ b/Engine/Plugins/Runtime/Interchange/Source/Parsers/Fbx/Private/FbxScene.cpp
@@ -7,6 +7,7 @@
 #include "FbxHelper.h"
 #include "FbxInclude.h"
 #include "FbxMaterial.h"
+#include "FbxMesh.h"
 #include "InterchangeSceneNode.h"
 #include "Nodes/InterchangeBaseNodeContainer.h"
 
@@ -18,51 +19,130 @@ namespace UE
 		{
 			namespace Scene
 			{
+				void CreateMeshNodeReference(UInterchangeSceneNode* UnrealSceneNode, FbxNode* FbxSceneNode, FbxNodeAttribute* NodeAttribute, UInterchangeBaseNodeContainer& NodeContainer, TArray<FString>& JSonErrorMessages)
+				{
+					UInterchangeBaseNode* MeshNode = nullptr;
+					if (NodeAttribute->GetAttributeType() == FbxNodeAttribute::eMesh)
+					{
+						FbxMesh* Mesh = static_cast<FbxMesh*>(NodeAttribute);
+						if (ensure(Mesh))
+						{
+							FString MeshRefString = FFbxMesh::GetMeshUniqueID(Mesh);
+							MeshNode = NodeContainer.GetNode(MeshRefString);
+						}
+					}
+					else if (NodeAttribute->GetAttributeType() == FbxNodeAttribute::eShape)
+					{
+						//We do not add a dependency for shape on the scene node since shapes are a MeshNode dependency.
+					}
+
+					if (MeshNode)
+					{
+						UnrealSceneNode->AddAssetDependency(MeshNode->GetUniqueID());
+					}
+				}
+
+				void CreateCameraNodeReference(UInterchangeSceneNode* UnrealSceneNode, FbxNode* FbxSceneNode, FbxNodeAttribute* NodeAttribute, UInterchangeBaseNodeContainer& NodeContainer, TArray<FString>& JSonErrorMessages)
+				{
+				}
+
+				void CreateLightNodeReference(UInterchangeSceneNode* UnrealSceneNode, FbxNode* FbxSceneNode, FbxNodeAttribute* NodeAttribute, UInterchangeBaseNodeContainer& NodeContainer, TArray<FString>& JSonErrorMessages)
+				{
+
+				}
+
 				UInterchangeSceneNode* AddHierarchyRecursively(FbxNode* Node, FbxScene* SDKScene, UInterchangeBaseNodeContainer& NodeContainer, TArray<FString>& JSonErrorMessages)
 				{
 					FString NodeName = FFbxHelper::GetFbxObjectName(Node);
 					FString NodeUniqueID = FFbxHelper::GetFbxNodeHierarchyName(Node);
+
+					UInterchangeSceneNode* UnrealNode = FFbxScene::CreateTransformNode(NodeContainer, NodeName, NodeUniqueID, JSonErrorMessages);
+					check(UnrealNode);
 					
-					UInterchangeSceneNode* UnrealNode = nullptr;
+					FTransform GlobalTransform = FFbxConvert::ConvertTransform(Node->EvaluateGlobalTransform());
+					UnrealNode->SetCustomGlobalTransform(GlobalTransform);
+					FbxNode* FbxParentNode = Node->GetParent();
+					FTransform LocalTransform;
+					if (FbxParentNode)
+					{
+						FTransform ParentGlobalTransform = FFbxConvert::ConvertTransform(FbxParentNode->EvaluateGlobalTransform());
+						LocalTransform = ParentGlobalTransform.Inverse() * GlobalTransform;
+					}
+					else
+					{
+						LocalTransform = GlobalTransform;
+					}
+					UnrealNode->SetCustomLocalTransform(LocalTransform);
 
 					int32 AttributeCount = Node->GetNodeAttributeCount();
 					for (int32 AttributeIndex = 0; AttributeIndex < AttributeCount; ++AttributeIndex)
 					{
-						const FbxNodeAttribute* NodeAttribute = Node->GetNodeAttributeByIndex(AttributeIndex);
+						FbxNodeAttribute* NodeAttribute = Node->GetNodeAttributeByIndex(AttributeIndex);
 						switch (NodeAttribute->GetAttributeType())
 						{
-							case FbxNodeAttribute::eLODGroup:
-							case FbxNodeAttribute::eMesh:
+							case FbxNodeAttribute::eUnknown:
+							case FbxNodeAttribute::eOpticalReference:
+							case FbxNodeAttribute::eOpticalMarker:
+							case FbxNodeAttribute::eCachedEffect:
+							case FbxNodeAttribute::eNull:
+							case FbxNodeAttribute::eMarker:
+							case FbxNodeAttribute::eCameraStereo:
+							case FbxNodeAttribute::eCameraSwitcher:
+							case FbxNodeAttribute::eNurbs:
+							case FbxNodeAttribute::ePatch:
+							case FbxNodeAttribute::eNurbsCurve:
+							case FbxNodeAttribute::eTrimNurbsSurface:
+							case FbxNodeAttribute::eBoundary:
+							case FbxNodeAttribute::eNurbsSurface:
+							case FbxNodeAttribute::eSubDiv:
+							case FbxNodeAttribute::eLine:
+								//Unsupported attribute
+								break;
+							case FbxNodeAttribute::eShape: //We do not add a dependency for shape on the scene node since shapes are a MeshNode dependency.
+								break;
 							case FbxNodeAttribute::eSkeleton:
+							{
+								//Add the joint specialized type
+								FString SpecializedType = TEXT("Joint");
+								UnrealNode->AddSpecializedType(SpecializedType);
+								break;
+							}
+							case FbxNodeAttribute::eMesh:
+							{
+								//For Mesh attribute we add the fbx nodes materials
+								FFbxMaterial::AddAllNodeMaterials(UnrealNode, Node, NodeContainer, JSonErrorMessages);
+								CreateMeshNodeReference(UnrealNode, Node, NodeAttribute, NodeContainer, JSonErrorMessages);
+								break;
+							}
+							case FbxNodeAttribute::eLODGroup:
+							{
+								FString SpecializedType = TEXT("LodGroup");
+								UnrealNode->AddSpecializedType(SpecializedType);
+								break;
+							}
 							case FbxNodeAttribute::eCamera:
-							case FbxNodeAttribute::eCameraSwitcher:
+							{
+								//Add the Camera asset
+								CreateCameraNodeReference(UnrealNode, Node, NodeAttribute, NodeContainer, JSonErrorMessages);
+								break;
+							}
 							case FbxNodeAttribute::eLight:
-							default:
-								//Add a transform node
-								UnrealNode = FFbxScene::CreateTransformNode(NodeContainer, NodeName, NodeUniqueID, JSonErrorMessages);
+							{
+								//Add the Light asset
+								CreateLightNodeReference(UnrealNode, Node, NodeAttribute, NodeContainer, JSonErrorMessages);
 								break;
+							}
 						}
 					}
-					if (UnrealNode)
-					{
-						NodeContainer.AddNode(UnrealNode);
-						FTransform LocalTransform = FFbxConvert::ConvertTransform(Node->EvaluateLocalTransform());
-						UnrealNode->SetCustomLocalTransform(LocalTransform);
-						FTransform GlobalTransform = FFbxConvert::ConvertTransform(Node->EvaluateGlobalTransform());
-						UnrealNode->SetCustomGlobalTransform(GlobalTransform);
-						NodeContainer.AddNode(UnrealNode);
-						//Add the dependencies of the material on the correct order
-						FFbxMaterial::AddAllNodeMaterials(UnrealNode, Node, NodeContainer, JSonErrorMessages);
 
-						const int32 ChildCount = Node->GetChildCount();
-						for (int32 ChildIndex = 0; ChildIndex < ChildCount; ++ChildIndex)
+					const int32 ChildCount = Node->GetChildCount();
+					for (int32 ChildIndex = 0; ChildIndex < ChildCount; ++ChildIndex)
+					{
+						FbxNode* ChildNode = Node->GetChild(ChildIndex);
+						UInterchangeSceneNode* UnrealChildNode = AddHierarchyRecursively(ChildNode, SDKScene, NodeContainer, JSonErrorMessages);
+						if (UnrealChildNode)
 						{
-							FbxNode* ChildNode = Node->GetChild(ChildIndex);
-							UInterchangeSceneNode* UnrealChildNode = AddHierarchyRecursively(ChildNode, SDKScene, NodeContainer, JSonErrorMessages);
-							if (UnrealChildNode)
-							{
-								UnrealChildNode->SetParentUID(UnrealNode->GetUniqueID());
-							}
+							UnrealChildNode->SetParentUid(UnrealNode->GetUniqueID());
 						}
 					}
 					return UnrealNode;
@@ -72,7 +152,7 @@ namespace UE
 			UInterchangeSceneNode* FFbxScene::CreateTransformNode(UInterchangeBaseNodeContainer& NodeContainer, const FString& NodeName, const FString& NodeUniqueID, TArray<FString>& JSonErrorMessages)
 			{
 				FString DisplayLabel(NodeName);
-				FString NodeUID(NodeUniqueID);
+				FString NodeUid(NodeUniqueID);
 				UInterchangeSceneNode* TransformNode = NewObject<UInterchangeSceneNode>(&NodeContainer, NAME_None);
 				if (!ensure(TransformNode))
 				{
@@ -80,7 +160,7 @@ namespace UE
 					return nullptr;
 				}
 				// Creating a UMaterialInterface
-				TransformNode->InitializeNode(NodeUID, DisplayLabel);
+				TransformNode->InitializeNode(NodeUid, DisplayLabel, EInterchangeNodeContainerType::NodeContainerType_TranslatedScene);
 				NodeContainer.AddNode(TransformNode);
 				return TransformNode;
 			}
diff --git a/Engine/Plugins/Runtime/Interchange/Source/Parsers/Fbx/Private/FbxSkeletalMesh.cpp b/Engine/Plugins/Runtime/Interchange/Source/Parsers/Fbx/Private/FbxSkeletalMesh.cpp
deleted file mode 100644
diff --git a/Engine/Plugins/Runtime/Interchange/Source/Parsers/Fbx/Private/FbxSkeletalMesh.h b/Engine/Plugins/Runtime/Interchange/Source/Parsers/Fbx/Private/FbxSkeletalMesh.h
deleted file mode 100644
diff --git a/Engine/Plugins/Runtime/Interchange/Source/Parsers/Fbx/Private/FbxSkeleton.cpp b/Engine/Plugins/Runtime/Interchange/Source/Parsers/Fbx/Private/FbxSkeleton.cpp
deleted file mode 100644
diff --git a/Engine/Plugins/Runtime/Interchange/Source/Parsers/Fbx/Private/FbxSkeleton.h b/Engine/Plugins/Runtime/Interchange/Source/Parsers/Fbx/Private/FbxSkeleton.h
deleted file mode 100644
diff --git a/Engine/Plugins/Runtime/Interchange/Source/Pipelines/Private/InterchangeGenericAssetsPipeline.cpp b/Engine/Plugins/Runtime/Interchange/Source/Pipelines/Private/InterchangeGenericAssetsPipeline.cpp
new file mode 100644
--- /dev/null
+++ b/Engine/Plugins/Runtime/Interchange/Source/Pipelines/Private/InterchangeGenericAssetsPipeline.cpp
@@ -0,0 +1,168 @@
+// Copyright Epic Games, Inc. All Rights Reserved. 
+#include "InterchangeGenericAssetsPipeline.h"
+
+#include "CoreMinimal.h"
+#include "InterchangeMaterialFactoryNode.h"
+#include "InterchangeMaterialNode.h"
+#include "InterchangeMeshNode.h"
+#include "InterchangePipelineLog.h"
+#include "InterchangeSceneNode.h"
+#include "InterchangeTextureFactoryNode.h"
+#include "InterchangeTextureNode.h"
+#include "InterchangeSourceData.h"
+#include "Misc/Paths.h"
+#include "Nodes/InterchangeBaseNodeContainer.h"
+#include "UObject/Object.h"
+#include "UObject/ObjectMacros.h"
+
+bool UInterchangeGenericAssetsPipeline::ExecutePreImportPipeline(UInterchangeBaseNodeContainer* InBaseNodeContainer, const TArray<UInterchangeSourceData*>& InSourceDatas)
+{
+	if (!InBaseNodeContainer)
+	{
+		UE_LOG(LogInterchangePipeline, Warning, TEXT("UInterchangeGenericAssetsPipeline: Cannot execute pre-import pipeline because InBaseNodeContrainer is null"));
+		return false;
+	}
+
+	BaseNodeContainer = InBaseNodeContainer;
+	SourceDatas.Empty(InSourceDatas.Num());
+	for (const UInterchangeSourceData* SourceData : InSourceDatas)
+	{
+		SourceDatas.Add(SourceData);
+	}
+
+	
+	//Find all translated node we need for this pipeline
+	BaseNodeContainer->IterateNodes([this](const FString& NodeUid, UInterchangeBaseNode* Node)
+	{
+		switch(Node->GetnodeContainerType())
+		{
+			case EInterchangeNodeContainerType::NodeContainerType_TranslatedAsset:
+			{
+				if (UInterchangeTextureNode* TextureNode = Cast<UInterchangeTextureNode>(Node))
+				{
+					TextureNodes.Add(TextureNode);
+				}
+				else if (UInterchangeMaterialNode* MaterialNode = Cast<UInterchangeMaterialNode>(Node))
+				{
+					MaterialNodes.Add(MaterialNode);
+				}
+				else if (UInterchangeMeshNode* MeshNode = Cast<UInterchangeMeshNode>(Node))
+				{
+					MeshNodes.Add(MeshNode);
+				}
+			}
+			break;
+			case EInterchangeNodeContainerType::NodeContainerType_TranslatedScene:
+			{
+				if (UInterchangeSceneNode* SceneNode = Cast<UInterchangeSceneNode>(Node))
+				{
+					SceneNodes.Add(SceneNode);
+				}
+			}
+			break;
+		}
+	});
+	
+	//import textures
+	for (const TObjectPtr<UInterchangeTextureNode> TextureNode : TextureNodes)
+	{
+		CreateTexture2DFactoryNode(TextureNode);
+	}
+	//import materials
+	for (const TObjectPtr<UInterchangeMaterialNode> MaterialNode : MaterialNodes)
+	{
+		CreateMaterialFactoryNode(MaterialNode);
+	}
+
+	//Find all skinned mesh
+	//import meshes
+	TArray<UInterchangeMeshNode*> SkinnedMeshNodes;
+	for (const TObjectPtr<UInterchangeMeshNode> MeshNode : MeshNodes)
+	{
+		if (MeshNode->IsSkinnedMesh())
+		{
+			SkinnedMeshNodes.Add(MeshNode.Get());
+		}
+	}
+
+	return true;
+}
+
+TObjectPtr<UInterchangeTextureFactoryNode> UInterchangeGenericAssetsPipeline::CreateTexture2DFactoryNode(const TObjectPtr<UInterchangeTextureNode> TextureNode)
+{
+	FString DisplayLabel = TextureNode->GetDisplayLabel();
+	FString NodeUid = UInterchangeTextureFactoryNode::GetTextureFactoryNodeUidFromTextureNodeUid(TextureNode->GetUniqueID());
+	UInterchangeTextureFactoryNode* TextureFactoryNode = nullptr;
+	if (BaseNodeContainer->IsNodeUidValid(NodeUid))
+	{
+		TextureFactoryNode = Cast<UInterchangeTextureFactoryNode>(BaseNodeContainer->GetNode(NodeUid));
+		if (!ensure(TextureFactoryNode))
+		{
+			//Log an error
+		}
+	}
+	else
+	{
+		TextureFactoryNode = NewObject<UInterchangeTextureFactoryNode>(BaseNodeContainer.Get(), NAME_None);
+		if (!ensure(TextureFactoryNode))
+		{
+			return nullptr;
+		}
+		//Creating a UTexture2D
+		TextureFactoryNode->InitializeTextureNode(NodeUid, DisplayLabel, UTexture2D::StaticClass()->GetName(), TextureNode->GetDisplayLabel());
+		TextureFactoryNode->SetCustomTranslatedTextureNodeUid(TextureNode->GetUniqueID());
+		BaseNodeContainer->AddNode(TextureFactoryNode);
+		TextureFactoryNodes.Add(TextureFactoryNode);
+	}
+	return TextureFactoryNode;
+}
+
+TObjectPtr<UInterchangeMaterialFactoryNode> UInterchangeGenericAssetsPipeline::CreateMaterialFactoryNode(const TObjectPtr<UInterchangeMaterialNode> MaterialNode)
+{
+	FString DisplayLabel = MaterialNode->GetDisplayLabel();
+	FString NodeUid = UInterchangeMaterialFactoryNode::GetMaterialFactoryNodeUidFromMaterialNodeUid(MaterialNode->GetUniqueID());
+	UInterchangeMaterialFactoryNode* MaterialFactoryNode = nullptr;
+	if (BaseNodeContainer->IsNodeUidValid(NodeUid))
+	{
+		//The node already exist, just return it
+		MaterialFactoryNode = Cast<UInterchangeMaterialFactoryNode>(BaseNodeContainer->GetNode(NodeUid));
+		if (!ensure(MaterialFactoryNode))
+		{
+			//Log an error
+		}
+	}
+	else
+	{
+		MaterialFactoryNode = NewObject<UInterchangeMaterialFactoryNode>(BaseNodeContainer.Get(), NAME_None);
+		if (!ensure(MaterialFactoryNode))
+		{
+			return nullptr;
+		}
+		//Creating a Material
+		MaterialFactoryNode->InitializeMaterialNode(NodeUid, DisplayLabel, UMaterial::StaticClass()->GetName());
+		MaterialFactoryNode->SetCustomTranslatedMaterialNodeUid(MaterialNode->GetUniqueID());
+		TArray<FString> TranslatedDependencies;
+		MaterialNode->GetTextureDependencies(TranslatedDependencies);
+		for (const FString& DepUid : TranslatedDependencies)
+		{
+			FString FactoryDepUid = UInterchangeTextureFactoryNode::GetTextureFactoryNodeUidFromTextureNodeUid(DepUid);
+			if (UInterchangeTextureFactoryNode* TextureDep = Cast<UInterchangeTextureFactoryNode>(BaseNodeContainer->GetNode(FactoryDepUid)))
+			{
+				//We found a factory dependency, because we need to be able to retrieve the UObject create by the factory
+				MaterialFactoryNode->SetTextureDependencyUid(FactoryDepUid);
+				//We also add it to the base dependencies so the texture factory will be called before this material factory
+				MaterialFactoryNode->SetFactoryDependencyUid(FactoryDepUid);
+			}
+		}
+		BaseNodeContainer->AddNode(MaterialFactoryNode);
+		MaterialFactoryNodes.Add(MaterialFactoryNode);
+	}
+	return MaterialFactoryNode;
+}
+
+TObjectPtr<UInterchangeBaseNode> UInterchangeGenericAssetsPipeline::CreateMeshFactoryNode(const TObjectPtr<UInterchangeMeshNode> MeshNode)
+{
+	FString DisplayLabel = MeshNode->GetDisplayLabel();
+	FString NodeUid = MeshNode->GetUniqueID();
+	return nullptr;
+}
\ No newline at end of file
diff --git a/Engine/Plugins/Runtime/Interchange/Source/Pipelines/Private/InterchangeGenericMeshPipeline.cpp b/Engine/Plugins/Runtime/Interchange/Source/Pipelines/Private/InterchangeGenericMeshPipeline.cpp
deleted file mode 100644
diff --git a/Engine/Plugins/Runtime/Interchange/Source/Pipelines/Private/InterchangePipelineLog.h b/Engine/Plugins/Runtime/Interchange/Source/Pipelines/Private/InterchangePipelineLog.h
new file mode 100644
--- /dev/null
+++ b/Engine/Plugins/Runtime/Interchange/Source/Pipelines/Private/InterchangePipelineLog.h
@@ -0,0 +1,7 @@
+// Copyright Epic Games, Inc. All Rights Reserved.
+
+#pragma once
+
+#include "CoreMinimal.h"
+
+DECLARE_LOG_CATEGORY_EXTERN(LogInterchangePipeline, Log, All);
diff --git a/Engine/Plugins/Runtime/Interchange/Source/Pipelines/Private/InterchangePipelinesModule.cpp b/Engine/Plugins/Runtime/Interchange/Source/Pipelines/Private/InterchangePipelinesModule.cpp
--- a/Engine/Plugins/Runtime/Interchange/Source/Pipelines/Private/InterchangePipelinesModule.cpp
+++ b/Engine/Plugins/Runtime/Interchange/Source/Pipelines/Private/InterchangePipelinesModule.cpp
@@ -4,11 +4,13 @@
 
 #include "CoreMinimal.h"
 #include "Engine/Engine.h"
-#include "InterchangeGenericMeshPipeline.h"
+//#include "InterchangeGenericAssetsPipeline.h"
 #include "InterchangeManager.h"
+#include "InterchangePipelineLog.h"
 #include "Misc/CoreDelegates.h"
 #include "Modules/ModuleManager.h"
 
+DEFINE_LOG_CATEGORY(LogInterchangePipeline);
 
 class FInterchangePipelinesModule : public IInterchangePipelinesModule
 {
diff --git a/Engine/Plugins/Runtime/Interchange/Source/Pipelines/Public/InterchangeGenericAssetsPipeline.h b/Engine/Plugins/Runtime/Interchange/Source/Pipelines/Public/InterchangeGenericAssetsPipeline.h
new file mode 100644
--- /dev/null
+++ b/Engine/Plugins/Runtime/Interchange/Source/Pipelines/Public/InterchangeGenericAssetsPipeline.h
@@ -0,0 +1,95 @@
+// Copyright Epic Games, Inc. All Rights Reserved.
+
+#pragma once
+
+#include "CoreMinimal.h"
+#include "InterchangePipelineBase.h"
+#include "InterchangeSourceData.h"
+#include "Nodes/InterchangeBaseNodeContainer.h"
+#include "UObject/Object.h"
+#include "UObject/ObjectMacros.h"
+
+#include "InterchangeGenericAssetsPipeline.generated.h"
+
+class UInterchangeTextureNode;
+class UInterchangeTextureFactoryNode;
+class UInterchangeMaterialNode;
+class UInterchangeMaterialFactoryNode;
+class UInterchangeMeshNode;
+class UInterchangeSceneNode;
+
+/**
+ * This pipeline is the generic pipeline option for all meshes type and should be call before specialized Mesh pipeline (like generic static mesh or skeletal mesh pipelines)
+ * All shared import options between mesh type should be added here.
+ *
+ */
+UCLASS(BlueprintType)
+class INTERCHANGEPIPELINES_API UInterchangeGenericAssetsPipeline : public UInterchangePipelineBase
+{
+	GENERATED_BODY()
+
+public:
+	//////////////////////////////////////////////////////////////////////////
+	// BEGIN Pre import pipeline properties
+
+	/** If enable and there is only one mesh in the translated nodes, the Source data name (filename) will be use to name the asset */
+	//UPROPERTY(EditAnywhere, BlueprintReadWrite, Category = "Miscellaneous")
+	//bool bUseSourceDataToNameMeshes = true;
+
+	// END Pre import pipeline properties
+	//////////////////////////////////////////////////////////////////////////
+
+protected:
+
+	virtual bool ExecutePreImportPipeline(UInterchangeBaseNodeContainer* InBaseNodeContainer, const TArray<UInterchangeSourceData*>& InSourceDatas) override;
+
+	//virtual bool ExecutePostImportPipeline(const UInterchangeBaseNodeContainer* BaseNodeContainer, const FName& NodeKey, UObject* CreatedAsset) override;
+
+	virtual bool CanExecuteOnAnyThread(EInterchangePipelineTask PipelineTask) override
+	{
+		//If a blueprint or python derived from this class, it will be execute on the game thread since we cannot currently execute script outside of the game thread, even if this function return true.
+		return true;
+	}
+
+	//virtual bool ExecuteExportPipeline(UInterchangeBaseNodeContainer* BaseNodeContainer) override;
+private:
+
+	TObjectPtr<UInterchangeBaseNodeContainer> BaseNodeContainer;
+	TArray<const UInterchangeSourceData*> SourceDatas;
+
+	
+	/** Texture translated assets nodes */
+	TArray<TObjectPtr<UInterchangeTextureNode>> TextureNodes;
+
+	/** Texture factory assets nodes */
+	TArray<TObjectPtr<UInterchangeTextureFactoryNode>> TextureFactoryNodes;
+
+	TObjectPtr<UInterchangeTextureFactoryNode> CreateTexture2DFactoryNode(const TObjectPtr<UInterchangeTextureNode> TextureNode);
+
+	
+	/** Material translated assets nodes */
+	TArray<TObjectPtr<UInterchangeMaterialNode>> MaterialNodes;
+	
+	/** Material factory assets nodes */
+	TArray<TObjectPtr<UInterchangeMaterialFactoryNode>> MaterialFactoryNodes;
+	
+	TObjectPtr<UInterchangeMaterialFactoryNode> CreateMaterialFactoryNode(const TObjectPtr<UInterchangeMaterialNode> MaterialNode);
+
+	
+	/** Mesh translated assets nodes */
+	TArray<TObjectPtr<UInterchangeMeshNode>> MeshNodes;
+	
+	/** Skeletal mesh factory assets nodes */
+	//TArray<TObjectPtr<UInterchangeSkeletalMeshFactoryNode>> SkeletalMeshFactoryNodes;
+	
+	/** Static mesh factory assets nodes */
+	//TArray<TObjectPtr<UInterchangeStaticMeshFactoryNode>> StaticMeshFactoryNodes;
+	
+	/** This function can create a UInterchangeSkeletalMeshFactoryNode or a UInterchangeStaticMeshFactoryNode */
+	TObjectPtr<UInterchangeBaseNode> CreateMeshFactoryNode(const TObjectPtr<UInterchangeMeshNode> MeshNode);
+
+	/** Mesh translated scene nodes */
+	TArray<TObjectPtr<UInterchangeSceneNode>> SceneNodes;
+};
+
+
diff --git a/Engine/Plugins/Runtime/Interchange/Source/Pipelines/Public/InterchangeGenericMeshPipeline.h b/Engine/Plugins/Runtime/Interchange/Source/Pipelines/Public/InterchangeGenericMeshPipeline.h
deleted file mode 100644
diff --git a/Engine/Source/Runtime/Interchange/Core/Private/Nodes/InterchangeBaseNode.cpp b/Engine/Source/Runtime/Interchange/Core/Private/Nodes/InterchangeBaseNode.cpp
--- a/Engine/Source/Runtime/Interchange/Core/Private/Nodes/InterchangeBaseNode.cpp
+++ b/Engine/Source/Runtime/Interchange/Core/Private/Nodes/InterchangeBaseNode.cpp
@@ -10,7 +10,7 @@
 
 
 
-void UInterchangeBaseNode::InitializeNode(const FString& UniqueID, const FString& DisplayLabel)
+void UInterchangeBaseNode::InitializeNode(const FString& UniqueID, const FString& DisplayLabel, const EInterchangeNodeContainerType NodeContainerType)
 {
 	UE::Interchange::EAttributeStorageResult Result = Attributes->RegisterAttribute(UE::Interchange::FBaseNodeStaticData::UniqueIDKey(), UniqueID, UE::Interchange::EAttributeProperty::NoHash);
 	if (!IsAttributeStorageResultSuccess(Result))
@@ -24,7 +24,13 @@ void UInterchangeBaseNode::InitializeNode(const FString& UniqueID, const FString
 		LogAttributeStorageErrors(Result, TEXT("RegisterAttribute"), UE::Interchange::FBaseNodeStaticData::DisplayLabelKey());
 	}
 
-	Dependencies.Initialize(Attributes, UE::Interchange::FBaseNodeStaticData::GetDependenciesBaseKey());
+	Result = Attributes->RegisterAttribute(UE::Interchange::FBaseNodeStaticData::NodeContainerTypeKey(), static_cast<uint8>(NodeContainerType));
+	if (!IsAttributeStorageResultSuccess(Result))
+	{
+		LogAttributeStorageErrors(Result, TEXT("RegisterAttribute"), UE::Interchange::FBaseNodeStaticData::NodeContainerTypeKey());
+	}
+
+	FactoryDependencies.Initialize(Attributes, UE::Interchange::FBaseNodeStaticData::GetFactoryDependenciesBaseKey());
 
 	bIsInitialized = true;
 }
@@ -69,11 +75,11 @@ bool UInterchangeBaseNode::SetDisplayLabel(const FString& DisplayLabel)
 	return false;
 }
 
-FString UInterchangeBaseNode::GetParentUID() const
+FString UInterchangeBaseNode::GetParentUid() const
 {
 	if (!Attributes->ContainAttribute(UE::Interchange::FBaseNodeStaticData::ParentIDKey()))
 	{
-		return InvalidNodeUID();
+		return InvalidNodeUid();
 	}
 
 	FString ParentUniqueID;
@@ -83,12 +89,12 @@ FString UInterchangeBaseNode::GetParentUID() const
 		Handle.Get(ParentUniqueID);
 		return ParentUniqueID;
 	}
-	return InvalidNodeUID();
+	return InvalidNodeUid();
 }
 
-bool UInterchangeBaseNode::SetParentUID(const FString& ParentUID)
+bool UInterchangeBaseNode::SetParentUid(const FString& ParentUid)
 {
-	UE::Interchange::EAttributeStorageResult Result = Attributes->RegisterAttribute(UE::Interchange::FBaseNodeStaticData::ParentIDKey(), ParentUID);
+	UE::Interchange::EAttributeStorageResult Result = Attributes->RegisterAttribute(UE::Interchange::FBaseNodeStaticData::ParentIDKey(), ParentUid);
 	if(IsAttributeStorageResultSuccess(Result))
 	{
 		UE::Interchange::FAttributeStorage::TAttributeHandle<FString> Handle = Attributes->GetAttributeHandle<FString>(UE::Interchange::FBaseNodeStaticData::ParentIDKey());
@@ -97,24 +103,29 @@ bool UInterchangeBaseNode::SetParentUID(const FString& ParentUID)
 	return false;
 }
 
-int32 UInterchangeBaseNode::GetDependeciesCount() const
+int32 UInterchangeBaseNode::GetFactoryDependenciesCount() const
 {
-	return Dependencies.GetCount();
+	return FactoryDependencies.GetCount();
 }
 
-void UInterchangeBaseNode::GetDependecies(TArray<FString>& OutDependencies ) const
+void UInterchangeBaseNode::GetFactoryDependency(const int32 Index, FString& OutDependency) const
 {
-	Dependencies.GetNames(OutDependencies);
+	FactoryDependencies.GetName(Index, OutDependency);
 }
 
-bool UInterchangeBaseNode::SetDependencyUID(const FString& DependencyUID)
+void UInterchangeBaseNode::GetFactoryDependencies(TArray<FString>& OutDependencies ) const
 {
-	return Dependencies.AddName(DependencyUID);
+	FactoryDependencies.GetNames(OutDependencies);
 }
 
-bool UInterchangeBaseNode::RemoveDependencyUID(const FString& DependencyUID)
+bool UInterchangeBaseNode::SetFactoryDependencyUid(const FString& DependencyUid)
 {
-	return Dependencies.RemoveName(DependencyUID);
+	return FactoryDependencies.AddName(DependencyUid);
+}
+
+bool UInterchangeBaseNode::RemoveFactoryDependencyUid(const FString& DependencyUid)
+{
+	return FactoryDependencies.RemoveName(DependencyUid);
 }
 
 bool UInterchangeBaseNode::IsEnabled() const
@@ -131,7 +142,7 @@ bool UInterchangeBaseNode::IsEnabled() const
 		Handle.Get(bIsEnabled);
 		return bIsEnabled;
 	}
-	return true;
+	return false;
 }
 
 bool UInterchangeBaseNode::SetEnabled(const bool bIsEnabled)
@@ -145,6 +156,22 @@ bool UInterchangeBaseNode::SetEnabled(const bool bIsEnabled)
 	return false;
 }
 
+EInterchangeNodeContainerType UInterchangeBaseNode::GetnodeContainerType() const
+{
+	if (!Attributes->ContainAttribute(UE::Interchange::FBaseNodeStaticData::NodeContainerTypeKey()))
+	{
+		return EInterchangeNodeContainerType::NodeContainerType_None;
+	}
+	UE::Interchange::FAttributeStorage::TAttributeHandle<uint8> Handle = Attributes->GetAttributeHandle<uint8>(UE::Interchange::FBaseNodeStaticData::NodeContainerTypeKey());
+	if (Handle.IsValid())
+	{
+		uint8 Value = static_cast<uint8>(EInterchangeNodeContainerType::NodeContainerType_None);
+		Handle.Get(Value);
+		return static_cast<EInterchangeNodeContainerType>(Value);
+	}
+	return EInterchangeNodeContainerType::NodeContainerType_None;
+}
+
 FGuid UInterchangeBaseNode::GetHash() const
 {
 	return Attributes->GetStorageHash();
@@ -155,7 +182,36 @@ class UClass* UInterchangeBaseNode::GetAssetClass() const
 	return nullptr;
 }
 
-FString UInterchangeBaseNode::InvalidNodeUID()
+FString UInterchangeBaseNode::GetAssetName() const
+{
+	if (!Attributes->ContainAttribute(UE::Interchange::FBaseNodeStaticData::AssetNameKey()))
+	{
+		return GetDisplayLabel();
+	}
+
+	UE::Interchange::FAttributeStorage::TAttributeHandle<FString> Handle = Attributes->GetAttributeHandle<FString>(UE::Interchange::FBaseNodeStaticData::AssetNameKey());
+	if (Handle.IsValid())
+	{
+		FString Value;
+		Handle.Get(Value);
+		return Value;
+	}
+
+	return GetDisplayLabel();
+}
+
+bool UInterchangeBaseNode::SetAssetName(const FString& AssetName)
+{
+	UE::Interchange::EAttributeStorageResult Result = Attributes->RegisterAttribute(UE::Interchange::FBaseNodeStaticData::AssetNameKey(), AssetName);
+	if (IsAttributeStorageResultSuccess(Result))
+	{
+		UE::Interchange::FAttributeStorage::TAttributeHandle<FString> Handle = Attributes->GetAttributeHandle<FString>(UE::Interchange::FBaseNodeStaticData::AssetNameKey());
+		return Handle.IsValid();
+	}
+	return false;
+}
+
+FString UInterchangeBaseNode::InvalidNodeUid()
 {
 	return FString();
 }
diff --git a/Engine/Source/Runtime/Interchange/Core/Private/Nodes/InterchangeBaseNodeContainer.cpp b/Engine/Source/Runtime/Interchange/Core/Private/Nodes/InterchangeBaseNodeContainer.cpp
--- a/Engine/Source/Runtime/Interchange/Core/Private/Nodes/InterchangeBaseNodeContainer.cpp
+++ b/Engine/Source/Runtime/Interchange/Core/Private/Nodes/InterchangeBaseNodeContainer.cpp
@@ -12,16 +12,21 @@
 #include "UObject/ObjectMacros.h"
 #include "UObject/UObjectIterator.h"
 
+UInterchangeBaseNodeContainer::UInterchangeBaseNodeContainer()
+{
+
+}
+
 FString UInterchangeBaseNodeContainer::AddNode(UInterchangeBaseNode* Node)
 {
 	if (!Node)
 	{
-		return UInterchangeBaseNode::InvalidNodeUID();
+		return UInterchangeBaseNode::InvalidNodeUid();
 	}
 	FString NodeUniqueID = Node->GetUniqueID();
-	if (NodeUniqueID == UInterchangeBaseNode::InvalidNodeUID())
+	if (NodeUniqueID == UInterchangeBaseNode::InvalidNodeUid())
 	{
-		return UInterchangeBaseNode::InvalidNodeUID();
+		return UInterchangeBaseNode::InvalidNodeUid();
 	}
 		
 	//Cannot add an node with the same IDs
@@ -43,9 +48,9 @@ FString UInterchangeBaseNodeContainer::AddNode(UInterchangeBaseNode* Node)
 	return NodeUniqueID;
 }
 
-bool UInterchangeBaseNodeContainer::IsNodeUIDValid(const FString& NodeUniqueID) const
+bool UInterchangeBaseNodeContainer::IsNodeUidValid(const FString& NodeUniqueID) const
 {
-	if (NodeUniqueID == UInterchangeBaseNode::InvalidNodeUID())
+	if (NodeUniqueID == UInterchangeBaseNode::InvalidNodeUid())
 	{
 		return false;
 	}
@@ -64,27 +69,28 @@ void UInterchangeBaseNodeContainer::GetRoots(TArray<FString>& RootNodes)
 {
 	for (auto& NodeKeyValue : Nodes)
 	{
-		if (NodeKeyValue.Value->GetParentUID() == UInterchangeBaseNode::InvalidNodeUID())
+		if (NodeKeyValue.Value->GetParentUid() == UInterchangeBaseNode::InvalidNodeUid())
 		{
 			RootNodes.Add(NodeKeyValue.Key);
 		}
 	}
 }
 
-void UInterchangeBaseNodeContainer::GetNodes(UClass* ClassNode, TArray<FString>& ClassNodes)
+void UInterchangeBaseNodeContainer::GetNodes(UClass* ClassNode, TArray<FString>& OutNodes)
 {
-	IterateNodes([&ClassNode, &ClassNodes](const FString& NodeUID, UInterchangeBaseNode* Node)
+	OutNodes.Empty();
+	IterateNodes([&ClassNode, &OutNodes](const FString& NodeUid, UInterchangeBaseNode* Node)
 	{
 		if(Node->GetClass()->IsChildOf(ClassNode))
 		{
-			ClassNodes.Add(Node->GetUniqueID());
+			OutNodes.Add(Node->GetUniqueID());
 		}
 	});
 }
 
 UInterchangeBaseNode* UInterchangeBaseNodeContainer::GetNode(const FString& NodeUniqueID)
 {
-	if (NodeUniqueID == UInterchangeBaseNode::InvalidNodeUID())
+	if (NodeUniqueID == UInterchangeBaseNode::InvalidNodeUid())
 	{
 		return nullptr;
 	}
@@ -101,38 +107,38 @@ const UInterchangeBaseNode* UInterchangeBaseNodeContainer::GetNode(const FString
 	return const_cast<UInterchangeBaseNodeContainer*>(this)->GetNode(NodeUniqueID);
 }
 
-bool UInterchangeBaseNodeContainer::SetNodeParentUID(const FString& NodeUniqueID, const FString& NewParentNodeUID)
+bool UInterchangeBaseNodeContainer::SetNodeParentUid(const FString& NodeUniqueID, const FString& NewParentNodeUid)
 {
 	if (!Nodes.Contains(NodeUniqueID))
 	{
 		return false;
 	}
-	if (!Nodes.Contains(NewParentNodeUID))
+	if (!Nodes.Contains(NewParentNodeUid))
 	{
 		return false;
 	}
 	UInterchangeBaseNode* Node = Nodes.FindChecked(NodeUniqueID);
-	Node->SetParentUID(NewParentNodeUID);
+	Node->SetParentUid(NewParentNodeUid);
 	return true;
 }
 
 int32 UInterchangeBaseNodeContainer::GetNodeChildrenCount(const FString& NodeUniqueID) const
 {
-	TArray<FString> ChildrenUIDs = GetNodeChildrenUIDs(NodeUniqueID);
-	return ChildrenUIDs.Num();
+	TArray<FString> ChildrenUids = GetNodeChildrenUids(NodeUniqueID);
+	return ChildrenUids.Num();
 }
 
-TArray<FString> UInterchangeBaseNodeContainer::GetNodeChildrenUIDs(const FString& NodeUniqueID) const
+TArray<FString> UInterchangeBaseNodeContainer::GetNodeChildrenUids(const FString& NodeUniqueID) const
 {
-	TArray<FString> OutChildrenUIDs;
+	TArray<FString> OutChildrenUids;
 	for (const auto& NodeKeyValue : Nodes)
 	{
-		if (NodeKeyValue.Value->GetParentUID() == NodeUniqueID)
+		if (NodeKeyValue.Value->GetParentUid() == NodeUniqueID)
 		{
-			OutChildrenUIDs.Add(NodeKeyValue.Key);
+			OutChildrenUids.Add(NodeKeyValue.Key);
 		}
 	}
-	return OutChildrenUIDs;
+	return OutChildrenUids;
 }
 
 UInterchangeBaseNode* UInterchangeBaseNodeContainer::GetNodeChildren(const FString& NodeUniqueID, int32 ChildIndex)
@@ -231,15 +237,15 @@ void UInterchangeBaseNodeContainer::LoadFromFile(const FString& Filename)
 
 UInterchangeBaseNode* UInterchangeBaseNodeContainer::GetNodeChildrenInternal(const FString& NodeUniqueID, int32 ChildIndex)
 {
-	TArray<FString> ChildrenUIDs = GetNodeChildrenUIDs(NodeUniqueID);
-	if (!ChildrenUIDs.IsValidIndex(ChildIndex))
+	TArray<FString> ChildrenUids = GetNodeChildrenUids(NodeUniqueID);
+	if (!ChildrenUids.IsValidIndex(ChildIndex))
 	{
 		return nullptr;
 	}
 
-	if (Nodes.Contains(ChildrenUIDs[ChildIndex]))
+	if (Nodes.Contains(ChildrenUids[ChildIndex]))
 	{
-		UInterchangeBaseNode* Node = Nodes.FindChecked(ChildrenUIDs[ChildIndex]);
+		UInterchangeBaseNode* Node = Nodes.FindChecked(ChildrenUids[ChildIndex]);
 		return Node;
 	}
 
diff --git a/Engine/Source/Runtime/Interchange/Core/Public/Nodes/InterchangeBaseNode.h b/Engine/Source/Runtime/Interchange/Core/Public/Nodes/InterchangeBaseNode.h
--- a/Engine/Source/Runtime/Interchange/Core/Public/Nodes/InterchangeBaseNode.h
+++ b/Engine/Source/Runtime/Interchange/Core/Public/Nodes/InterchangeBaseNode.h
@@ -751,10 +751,10 @@ namespace UE
 				return AttributeKey;
 			}
 
-			static const FString& GetDependenciesBaseKey()
+			static const FString& GetFactoryDependenciesBaseKey()
 			{
-				static FString BaseNodeDependencies_BaseKey = TEXT("__BaseNodeDependencies__");
-				return BaseNodeDependencies_BaseKey;
+				static FString BaseNodeFactoryDependencies_BaseKey = TEXT("__BaseNodeFactoryDependencies__");
+				return BaseNodeFactoryDependencies_BaseKey;
 			}
 
 			static const FAttributeKey& ClassTypeAttributeKey()
@@ -762,11 +762,32 @@ namespace UE
 				static FAttributeKey AttributeKey(TEXT("__ClassTypeAttribute__"));
 				return AttributeKey;
 			}
+
+			static const FAttributeKey& AssetNameKey()
+			{
+				static FAttributeKey AttributeKey(TEXT("__Asset_Name_Key__"));
+				return AttributeKey;
+			}
+
+			static const FAttributeKey& NodeContainerTypeKey()
+			{
+				static FAttributeKey AttributeKey(TEXT("__Node_Container_Type_Key__"));
+				return AttributeKey;
+			}
 		};
 
 	} //ns Interchange
 } //ns UE
 
+UENUM(BlueprintType)
+enum class EInterchangeNodeContainerType : uint8
+{
+	NodeContainerType_None,
+	NodeContainerType_TranslatedScene,
+	NodeContainerType_TranslatedAsset,
+	NodeContainerType_FactoryData
+};
+
 /**
  * This struct is used to store and retrieve key value attributes. The attributes are store in a generic FAttributeStorage which serialize the value in a TArray64<uint8>
  * See UE::Interchange::EAttributeTypes to know the supported template types
@@ -781,8 +802,9 @@ public:
 	UInterchangeBaseNode()
 	{
 		Attributes = MakeShared<UE::Interchange::FAttributeStorage, ESPMode::ThreadSafe>();
-		Dependencies.Initialize(Attributes, UE::Interchange::FBaseNodeStaticData::GetDependenciesBaseKey());
+		FactoryDependencies.Initialize(Attributes, UE::Interchange::FBaseNodeStaticData::GetFactoryDependenciesBaseKey());
 		RegisterAttribute<bool>(UE::Interchange::FBaseNodeStaticData::IsEnabledKey(), true);
+		RegisterAttribute<uint8>(UE::Interchange::FBaseNodeStaticData::NodeContainerTypeKey(), static_cast<uint8>(EInterchangeNodeContainerType::NodeContainerType_None));
 	}
 
 	virtual ~UInterchangeBaseNode() = default;
@@ -794,7 +816,7 @@ public:
 	 *
 	 */
 	UFUNCTION(BlueprintCallable, Category = "Interchange | Node")
-	void InitializeNode(const FString& UniqueID, const FString& DisplayLabel);
+	void InitializeNode(const FString& UniqueID, const FString& DisplayLabel, const EInterchangeNodeContainerType NodeContainerType);
 
 	/**
 	 * Return the node type name of the class, we use this when reporting error
@@ -832,13 +854,17 @@ public:
 		{
 			KeyDisplayName = TEXT("Node Class Type");
 		}
-		else if (NodeAttributeKey.Key.Equals(UE::Interchange::FBaseNodeStaticData::GetDependenciesBaseKey()))
+		else if (NodeAttributeKey == UE::Interchange::FBaseNodeStaticData::AssetNameKey())
 		{
-			KeyDisplayName = TEXT("Dependencies Count");
+			KeyDisplayName = TEXT("Imported Asset Name");
 		}
-		else if (NodeAttributeKey.Key.StartsWith(UE::Interchange::FBaseNodeStaticData::GetDependenciesBaseKey()))
+		else if (NodeAttributeKey.Key.Equals(UE::Interchange::FBaseNodeStaticData::GetFactoryDependenciesBaseKey()))
 		{
-			KeyDisplayName = TEXT("Dependencies Index ");
+			KeyDisplayName = TEXT("Factory Dependencies Count");
+		}
+		else if (NodeAttributeKey.Key.StartsWith(UE::Interchange::FBaseNodeStaticData::GetFactoryDependenciesBaseKey()))
+		{
+			KeyDisplayName = TEXT("Factory Dependencies Index ");
 			const FString IndexKey = UE::Interchange::FNameAttributeArrayHelper::IndexKey();
 			int32 IndexPosition = NodeAttributeKey.Key.Find(IndexKey) + IndexKey.Len();
 			if (IndexPosition < NodeAttributeKey.Key.Len())
@@ -846,9 +872,28 @@ public:
 				KeyDisplayName += NodeAttributeKey.Key.RightChop(IndexPosition);
 			}
 		}
+		else if (NodeAttributeKey == UE::Interchange::FBaseNodeStaticData::NodeContainerTypeKey())
+		{
+			KeyDisplayName = TEXT("Node Container Type");
+		}
 		return KeyDisplayName;
 	}
 
+	virtual bool ShouldHideAttribute(const UE::Interchange::FAttributeKey& NodeAttributeKey) const
+	{
+		return false;
+	}
+
+	virtual FString GetAttributeCategory(const UE::Interchange::FAttributeKey& NodeAttributeKey) const
+	{
+		FString CategoryName = TEXT("Attributes");
+		if (NodeAttributeKey.Key.StartsWith(UE::Interchange::FBaseNodeStaticData::GetFactoryDependenciesBaseKey()))
+		{
+			CategoryName = TEXT("FactoryDependencies");
+		}
+		return CategoryName;
+	}
+
 	/**
 	 * Add an attribute to the node
 	 * @param NodeAttributeKey - The key of the attribute
@@ -924,42 +969,49 @@ public:
 	bool SetDisplayLabel(const FString& DisplayName);
 
 	/**
-	 * Return the parent unique id. In case the attribute does not exist it will return InvalidNodeUID()
+	 * Return the parent unique id. In case the attribute does not exist it will return InvalidNodeUid()
 	 */
 	UFUNCTION(BlueprintCallable, Category = "Interchange | Node")
-	FString GetParentUID() const;
+	FString GetParentUid() const;
 
 	/**
 	 * Set the parent unique id.
 	 */
 	UFUNCTION(BlueprintCallable, Category = "Interchange | Node")
-	bool SetParentUID(const FString& ParentUID);
+	bool SetParentUid(const FString& ParentUid);
 
 	/**
-	 * This function allow to retrieve the number of dependencies for this object.
+	 * This function allow to retrieve the number of factory dependencies for this object.
 	 *
 	 */
 	UFUNCTION(BlueprintCallable, Category = "Interchange | Node")
-	int32 GetDependeciesCount() const;
+	int32 GetFactoryDependenciesCount() const;
 
 	/**
 	 * This function allow to retrieve the dependency for this object.
 	 * 
 	 */
 	UFUNCTION(BlueprintCallable, Category = "Interchange | Node")
-	void GetDependecies(TArray<FString>& OutDependencies ) const;
+	void GetFactoryDependencies(TArray<FString>& OutDependencies ) const;
+
+	/**
+	 * This function allow to retrieve one dependency for this object.
+	 *
+	 */
+	UFUNCTION(BlueprintCallable, Category = "Interchange | Node")
+	void GetFactoryDependency(const int32 Index, FString& OutDependency) const;
 
 	/**
 	 * Add one dependency to this object.
 	 */
 	UFUNCTION(BlueprintCallable, Category = "Interchange | Node")
-	bool SetDependencyUID(const FString& DependencyUID);
+	bool SetFactoryDependencyUid(const FString& DependencyUid);
 
 	/**
 	 * Remove one dependency from this object.
 	 */
 	UFUNCTION(BlueprintCallable, Category = "Interchange | Node")
-	bool RemoveDependencyUID(const FString& DependencyUID);
+	bool RemoveFactoryDependencyUid(const FString& DependencyUid);
 
 	/**
 	 * IsEnable true mean that the node will be import/export, if false it will be discarded.
@@ -976,6 +1028,12 @@ public:
 	UFUNCTION(BlueprintCallable, Category = "Interchange | Node")
 	bool SetEnabled(const bool bIsEnabled);
 
+	/**
+	 * Return the node container type which define the purpose of the node (Factory node, translated scene node or translated asset node).
+	 */
+	UFUNCTION(BlueprintCallable, Category = "Interchange | Node")
+	EInterchangeNodeContainerType GetnodeContainerType() const;
+
 	/**
 	 * Return a FGuid build from the FSHA1 of all the attribute data contain in the node.
 	 *
@@ -990,8 +1048,21 @@ public:
 	 */
 	virtual class UClass* GetAssetClass() const;
 
+	/**
+	 * Optional, Any node that can import/export an asset should set the proper name we will give to the asset.
+	 * If the attribute was never set, it will return GetDisplayLabel.
+	 */
+	UFUNCTION(BlueprintCallable, Category = "Interchange | Node")
+	virtual FString GetAssetName() const;
+
+	/**
+	 * Set the name we want for the imported asset this node represent. The asset factory will call GetAssetName()
+	 */
+	UFUNCTION(BlueprintCallable, Category = "Interchange | Node")
+	virtual bool SetAssetName(const FString& AssetName);
+
 	/** Return the invalid unique ID */
-	static FString InvalidNodeUID();
+	static FString InvalidNodeUid();
 
 	/**
 	 * Each Attribute that was set and have a delegate set for the specified UObject->UClass will
@@ -1033,5 +1104,10 @@ protected:
 
 	bool bIsInitialized = false;
 
-	UE::Interchange::FNameAttributeArrayHelper Dependencies;
+	/**
+	 * Those dependencies are use by the interchange parsing task to make sure the asset are created in the correct order.
+	 * Example: Mesh factory node will have dependencies on material factory node
+	 *          Material factory node will have dependencies on texture factory node
+	 */
+	UE::Interchange::FNameAttributeArrayHelper FactoryDependencies;
 };
diff --git a/Engine/Source/Runtime/Interchange/Core/Public/Nodes/InterchangeBaseNodeContainer.h b/Engine/Source/Runtime/Interchange/Core/Public/Nodes/InterchangeBaseNodeContainer.h
--- a/Engine/Source/Runtime/Interchange/Core/Public/Nodes/InterchangeBaseNodeContainer.h
+++ b/Engine/Source/Runtime/Interchange/Core/Public/Nodes/InterchangeBaseNodeContainer.h
@@ -24,12 +24,13 @@ class INTERCHANGECORE_API UInterchangeBaseNodeContainer : public UObject
 {
 	GENERATED_BODY()
 public:
+	UInterchangeBaseNodeContainer();
 
 	/**
 	 * Add a node in the container, the node will be add into a TMap.
 	 *
 	 * @param Node - a pointer on the node you want to add
-	 * @return: return the node unique ID of the added item. If the node already exist it will return the existing ID. Return InvalidNodeUID if the node cannot be added.
+	 * @return: return the node unique ID of the added item. If the node already exist it will return the existing ID. Return InvalidNodeUid if the node cannot be added.
 	 *
 	 */
 	UFUNCTION(BlueprintCallable, Category = "Interchange | Node Container")
@@ -37,7 +38,7 @@ public:
 
 	/** Return true if the node unique ID exist in the container */
 	UFUNCTION(BlueprintCallable, Category = "Interchange | Node Container")
-	bool IsNodeUIDValid(const FString& NodeUniqueID) const;
+	bool IsNodeUidValid(const FString& NodeUniqueID) const;
 
 	/** Unordered iteration of the all nodes */
 	void IterateNodes(TFunctionRef<void(const FString&, UInterchangeBaseNode*)> IterationLambda);
@@ -48,7 +49,7 @@ public:
 
 	/** Return all nodes that are of the ClassNode type*/
 	UFUNCTION(BlueprintCallable, Category = "Interchange | Node Container")
-	void GetNodes(UClass* ClassNode, TArray<FString>& ClassNodes);
+	void GetNodes(UClass* ClassNode, TArray<FString>& OutNodes);
 
 	/** Get an node pointer */
 	UFUNCTION(BlueprintCallable, Category = "Interchange | Node Container")
@@ -57,17 +58,17 @@ public:
 	/** Get an node pointer */
 	const UInterchangeBaseNode* GetNode(const FString& NodeUniqueID) const;
 
-	/** Set node ParentUID */
+	/** Set node ParentUid */
 	UFUNCTION(BlueprintCallable, Category = "Interchange | Node Container")
-	bool SetNodeParentUID(const FString& NodeUniqueID, const FString& NewParentNodeUID);
+	bool SetNodeParentUid(const FString& NodeUniqueID, const FString& NewParentNodeUid);
 
 	/** Get the node children count */
 	UFUNCTION(BlueprintCallable, Category = "Interchange | Node Container")
 	int32 GetNodeChildrenCount(const FString& NodeUniqueID) const;
 
-	/** Get all children UID */
+	/** Get all children Uid */
 	UFUNCTION(BlueprintCallable, Category = "Interchange | Node Container")
-	TArray<FString> GetNodeChildrenUIDs(const FString& NodeUniqueID) const;
+	TArray<FString> GetNodeChildrenUids(const FString& NodeUniqueID) const;
 
 	/** Get the node nth const children */
 	UFUNCTION(BlueprintCallable, Category = "Interchange | Node Container")
diff --git a/Engine/Source/Runtime/Interchange/Engine/Private/Tasks/InterchangeTaskCreateAsset.cpp b/Engine/Source/Runtime/Interchange/Engine/Private/Tasks/InterchangeTaskCreateAsset.cpp
--- a/Engine/Source/Runtime/Interchange/Engine/Private/Tasks/InterchangeTaskCreateAsset.cpp
+++ b/Engine/Source/Runtime/Interchange/Engine/Private/Tasks/InterchangeTaskCreateAsset.cpp
@@ -28,7 +28,7 @@ namespace UE
 			{
 				const UInterchangeSourceData* SourceData = AsyncHelper.SourceDatas[SourceIndex];
 				check(SourceData);
-				FString NodeDisplayName = Node->GetDisplayLabel();
+				FString NodeDisplayName = Node->GetAssetName();
 
 				// Set the asset name and the package name
 				OutAssetName = NodeDisplayName;
diff --git a/Engine/Source/Runtime/Interchange/Engine/Private/Tasks/InterchangeTaskParsing.cpp b/Engine/Source/Runtime/Interchange/Engine/Private/Tasks/InterchangeTaskParsing.cpp
--- a/Engine/Source/Runtime/Interchange/Engine/Private/Tasks/InterchangeTaskParsing.cpp
+++ b/Engine/Source/Runtime/Interchange/Engine/Private/Tasks/InterchangeTaskParsing.cpp
@@ -74,7 +74,7 @@ void UE::Interchange::FTaskParsing::DoTask(ENamedThreads::Type CurrentThread, co
 					NodeTaskData.UniqueID = Node->GetUniqueID();
 					NodeTaskData.SourceIndex = SourceIndex;
 					NodeTaskData.Node = Node;
-					Node->GetDependecies(NodeTaskData.Dependencies);
+					Node->GetFactoryDependencies(NodeTaskData.Dependencies);
 					NodeTaskData.Factory = NodeFactory;
 				}
 			});
