Generic pipeline can now import combined skeletal mesh asset.

diff --git a/Engine/Plugins/Runtime/Interchange/Source/Import/Private/Fbx/InterchangeFbxTranslator.cpp b/Engine/Plugins/Runtime/Interchange/Source/Import/Private/Fbx/InterchangeFbxTranslator.cpp
--- a/Engine/Plugins/Runtime/Interchange/Source/Import/Private/Fbx/InterchangeFbxTranslator.cpp
+++ b/Engine/Plugins/Runtime/Interchange/Source/Import/Private/Fbx/InterchangeFbxTranslator.cpp
@@ -200,16 +200,18 @@ TOptional<UE::Interchange::FSkeletalMeshLodPayloadData> UInterchangeFbxTranslato
 		return TOptional<UE::Interchange::FSkeletalMeshLodPayloadData>();
 	}
 
-
 	//Buffer keep the ownership of the data, the large memory reader is use to serialize the TMap
 	FLargeMemoryReader Ar(FileData, FileDataSize);
 	SkeletalMeshLodPayload.LodMeshDescription.Serialize(Ar);
-	//Read the blend shape
-	Ar << SkeletalMeshLodPayload.LodBlendShapeMeshDescriptions;
+	bool bFetchSkinnedData = false;
+	Ar << bFetchSkinnedData;
+	if (bFetchSkinnedData)
+	{
+		//Read the bone Name to remap the influence correctly
+		Ar << SkeletalMeshLodPayload.JointNames;
+	}
 
 	return SkeletalMeshLodPayload;
-
-	//return TOptional<UE::Interchange::FSkeletalMeshLodPayloadData>(SkeletalMeshLodPayload);
 }
 
 FString UInterchangeFbxTranslator::CreateLoadFbxFileCommand(const FString& FbxFilePath) const
diff --git a/Engine/Plugins/Runtime/Interchange/Source/Import/Private/Mesh/InterchangeSkeletalMeshFactory.cpp b/Engine/Plugins/Runtime/Interchange/Source/Import/Private/Mesh/InterchangeSkeletalMeshFactory.cpp
--- a/Engine/Plugins/Runtime/Interchange/Source/Import/Private/Mesh/InterchangeSkeletalMeshFactory.cpp
+++ b/Engine/Plugins/Runtime/Interchange/Source/Import/Private/Mesh/InterchangeSkeletalMeshFactory.cpp
@@ -6,7 +6,10 @@
 #include "GPUSkinPublicDefs.h"
 #include "InterchangeImportCommon.h"
 #include "InterchangeImportLog.h"
+#include "InterchangeMaterialFactoryNode.h"
 #include "InterchangeMaterialNode.h"
+#include "InterchangeMeshNode.h"
+#include "InterchangeSceneNode.h"
 #include "InterchangeSkeletalMeshLodDataNode.h"
 #include "InterchangeSkeletalMeshFactoryNode.h"
 #include "InterchangeSkeletonFactoryNode.h"
@@ -27,7 +30,7 @@
 
 #endif //WITH_EDITORONLY_DATA
 
-#if 0//WITH_EDITOR
+#if WITH_EDITOR
 namespace UE
 {
 	namespace Interchange
@@ -43,8 +46,8 @@ namespace UE
 
 			void RecursiveAddBones(const UInterchangeBaseNodeContainer* NodeContainer, const FString& JointNodeId, TArray <FJointInfo>& JointInfos, int32 ParentIndex, TArray<SkeletalMeshImportData::FBone>& RefBonesBinary)
 			{
-				const UInterchangeJointNode* JointNode = Cast<UInterchangeJointNode>(NodeContainer->GetNode(JointNodeId));
-				if (!JointNode)
+				const UInterchangeSceneNode* JointNode = Cast<UInterchangeSceneNode>(NodeContainer->GetNode(JointNodeId));
+				if (!JointNode || !JointNode->IsSpecializedTypeContains(FSceneNodeStaticData::GetJointSpecializeTypeString()))
 				{
 					UE_LOG(LogInterchangeImport, Warning, TEXT("Invalid Skeleton Joint"));
 					return;
@@ -52,7 +55,7 @@ namespace UE
 
 				int32 JointInfoIndex = JointInfos.Num();
 				FJointInfo& Info = JointInfos.AddZeroed_GetRef();
-				ensure(JointNode->GetCustomName(Info.Name));
+				Info.Name = JointNode->GetDisplayLabel();
 				ensure(JointNode->GetCustomLocalTransform(Info.LocalTransform));
 				Info.ParentIndex = ParentIndex;
 
@@ -66,7 +69,7 @@ namespace UE
 				Bone.BonePos.YSize = 1.0f;
 				Bone.BonePos.ZSize = 1.0f;
 				
-				const TArray<FString> ChildrenIds = NodeContainer->GetNodeChildrenUIDs(JointNodeId);
+				const TArray<FString> ChildrenIds = NodeContainer->GetNodeChildrenUids(JointNodeId);
 				Bone.NumChildren = ChildrenIds.Num();
 				for (int32 ChildIndex = 0; ChildIndex < ChildrenIds.Num(); ++ChildIndex)
 				{
@@ -178,144 +181,6 @@ namespace UE
 				}
 			}
 
-			void SetMaterialSkinXXOrder(FSkeletalMeshImportData& ImportData)
-			{
-				TArray<int32> MaterialIndexToSkinIndex;
-				TMap<int32, int32> SkinIndexToMaterialIndex;
-				TArray<int32> MissingSkinSuffixMaterial;
-				TMap<int32, int32> SkinIndexGreaterThenMaterialArraySize;
-				{
-					int32 MaterialCount = ImportData.Materials.Num();
-
-					bool bNeedsReorder = false;
-					for (int32 MaterialIndex = 0; MaterialIndex < MaterialCount; ++MaterialIndex)
-					{
-						// get skin index
-						FString MatName = ImportData.Materials[MaterialIndex].MaterialImportName;
-
-						if (MatName.Len() > 6)
-						{
-							int32 Offset = MatName.Find(TEXT("_SKIN"), ESearchCase::IgnoreCase, ESearchDir::FromEnd);
-							if (Offset != INDEX_NONE)
-							{
-								// Chop off the material name so we are left with the number in _SKINXX
-								FString SkinXXNumber = MatName.Right(MatName.Len() - (Offset + 1)).RightChop(4);
-
-								if (SkinXXNumber.IsNumeric())
-								{
-									bNeedsReorder = true;
-									int32 TmpIndex = FPlatformString::Atoi(*SkinXXNumber);
-									if (TmpIndex < MaterialCount)
-									{
-										SkinIndexToMaterialIndex.Add(TmpIndex, MaterialIndex);
-									}
-									else
-									{
-										SkinIndexGreaterThenMaterialArraySize.Add(TmpIndex, MaterialIndex);
-									}
-
-								}
-							}
-							else
-							{
-								MissingSkinSuffixMaterial.Add(MaterialIndex);
-							}
-						}
-						else
-						{
-							MissingSkinSuffixMaterial.Add(MaterialIndex);
-						}
-					}
-
-					if (bNeedsReorder && MissingSkinSuffixMaterial.Num() > 0)
-					{
-						//TODO log a user warning message
-						//AddTokenizedErrorMessage(FTokenizedMessage::Create(EMessageSeverity::Error, LOCTEXT("FbxSkeletaLMeshimport_Skinxx_missing", "Cannot mix skinxx suffix materials with no skinxx material, mesh section order will not be right.")), FFbxErrors::Generic_Mesh_SkinxxNameError);
-						return;
-					}
-
-					//Add greater then material array skinxx at the end sorted by integer the index will be remap correctly in the case of a LOD import
-					if (SkinIndexGreaterThenMaterialArraySize.Num() > 0)
-					{
-						int32 MaxAvailableKey = SkinIndexToMaterialIndex.Num();
-						for (int32 AvailableKey = 0; AvailableKey < MaxAvailableKey; ++AvailableKey)
-						{
-							if (SkinIndexToMaterialIndex.Contains(AvailableKey))
-								continue;
-
-							TMap<int32, int32> TempSkinIndexToMaterialIndex;
-							for (auto KvpSkinToMat : SkinIndexToMaterialIndex)
-							{
-								if (KvpSkinToMat.Key > AvailableKey)
-								{
-									TempSkinIndexToMaterialIndex.Add(KvpSkinToMat.Key - 1, KvpSkinToMat.Value);
-								}
-								else
-								{
-									TempSkinIndexToMaterialIndex.Add(KvpSkinToMat.Key, KvpSkinToMat.Value);
-								}
-							}
-							//move all the later key of the array to fill the available index
-							SkinIndexToMaterialIndex = TempSkinIndexToMaterialIndex;
-							AvailableKey--; //We need to retest the same index it can be empty
-						}
-						//Reorder the array
-						SkinIndexGreaterThenMaterialArraySize.KeySort(TLess<int32>());
-						for (auto Kvp : SkinIndexGreaterThenMaterialArraySize)
-						{
-							SkinIndexToMaterialIndex.Add(SkinIndexToMaterialIndex.Num(), Kvp.Value);
-						}
-					}
-
-					//Fill the array MaterialIndexToSkinIndex so we order material by _skinXX order
-					//This ensure we support skinxx suffixe that are not increment by one like _skin00, skin_01, skin_03, skin_04, skin_08... 
-					for (auto kvp : SkinIndexToMaterialIndex)
-					{
-						int32 MatIndexToInsert = 0;
-						for (MatIndexToInsert = 0; MatIndexToInsert < MaterialIndexToSkinIndex.Num(); ++MatIndexToInsert)
-						{
-							if (*(SkinIndexToMaterialIndex.Find(MaterialIndexToSkinIndex[MatIndexToInsert])) >= kvp.Value)
-							{
-								break;
-							}
-						}
-						MaterialIndexToSkinIndex.Insert(kvp.Key, MatIndexToInsert);
-					}
-
-					if (bNeedsReorder)
-					{
-						// re-order the materials
-						TArray< SkeletalMeshImportData::FMaterial > ExistingMatList = ImportData.Materials;
-						for (int32 MissingIndex : MissingSkinSuffixMaterial)
-						{
-							MaterialIndexToSkinIndex.Insert(MaterialIndexToSkinIndex.Num(), MissingIndex);
-						}
-						for (int32 MaterialIndex = 0; MaterialIndex < MaterialCount; ++MaterialIndex)
-						{
-							if (MaterialIndex < MaterialIndexToSkinIndex.Num())
-							{
-								int32 NewIndex = MaterialIndexToSkinIndex[MaterialIndex];
-								if (ExistingMatList.IsValidIndex(NewIndex))
-								{
-									ImportData.Materials[NewIndex] = ExistingMatList[MaterialIndex];
-								}
-							}
-						}
-
-						// remapping the material index for each triangle
-						int32 FaceNum = ImportData.Faces.Num();
-						for (int32 TriangleIndex = 0; TriangleIndex < FaceNum; TriangleIndex++)
-						{
-							SkeletalMeshImportData::FTriangle& Triangle = ImportData.Faces[TriangleIndex];
-							if (Triangle.MatIndex < MaterialIndexToSkinIndex.Num())
-							{
-								Triangle.MatIndex = MaterialIndexToSkinIndex[Triangle.MatIndex];
-							}
-						}
-					}
-				}
-			}
-
 			namespace SmoothGroupHelper
 			{
 				struct tFaceRecord
@@ -696,8 +561,6 @@ namespace UE
 				}
 
 				CleanUpUnusedMaterials(DestinationSkeletalMeshImportData);
-				// reorder material according to "SKinXX" in material name
-				SetMaterialSkinXXOrder(DestinationSkeletalMeshImportData);
 				DoUnSmoothVerts(DestinationSkeletalMeshImportData);
 
 				return true;
@@ -755,6 +618,7 @@ namespace UE
 			 * Fill the Materials array using the raw skeletalmesh geometry data (using material imported name)
 			 * Find the material from the dependencies of the skeletalmesh before searching in all package.
 			 */
+			//TODO: the pipeline should search for existing material and hook those before the factory is called
 			void ProcessImportMeshMaterials(TArray<FSkeletalMaterial>& Materials, FSkeletalMeshImportData& ImportData, TMap<FString, UMaterialInterface*>& AvailableMaterials)
 			{
 				TArray <SkeletalMeshImportData::FMaterial>& ImportedMaterials = ImportData.Materials;
@@ -766,19 +630,14 @@ namespace UE
 					const SkeletalMeshImportData::FMaterial& ImportedMaterial = ImportedMaterials[MatIndex];
 
 					UMaterialInterface* Material = nullptr;
-					//Remove _SkinXX from the material name
-					FString MaterialNameNoSkin = UE::Interchange::Material::RemoveSkinFromName(ImportedMaterial.MaterialImportName);
+
 					const FName SearchMaterialSlotName(*ImportedMaterial.MaterialImportName);
 					int32 MaterialIndex = 0;
 					FSkeletalMaterial* SkeletalMeshMaterialFind = Materials.FindByPredicate([&SearchMaterialSlotName, &MaterialIndex](const FSkeletalMaterial& ItemMaterial)
 					{
 						//Imported material slot name is available only WITH_EDITOR
 						FName ImportedMaterialSlot = NAME_None;
-#if WITH_EDITOR
 						ImportedMaterialSlot = ItemMaterial.ImportedMaterialSlotName;
-#else
-						ImportedMaterialSlot = ItemMaterial.MaterialSlotName;
-#endif
 						if (ImportedMaterialSlot != SearchMaterialSlotName)
 						{
 							MaterialIndex++;
@@ -795,26 +654,20 @@ namespace UE
 					if(!Material)
 					{
 						//Try to find the material in the skeletal mesh node dependencies (Materials are import before skeletal mesh when there is a dependency)
-						const FString MaterialName = ImportedMaterial.MaterialImportName;
-						if (AvailableMaterials.Contains(MaterialName))
+						if (AvailableMaterials.Contains(ImportedMaterial.MaterialImportName))
 						{
-							Material = AvailableMaterials.FindChecked(MaterialName);
+							Material = AvailableMaterials.FindChecked(ImportedMaterial.MaterialImportName);
 						}
 						else
 						{
 							//We did not found any material in the dependencies so try to find material everywhere
-							Material = FindObject<UMaterialInterface>(ANY_PACKAGE, *MaterialName);
-							if (Material == nullptr && MaterialNameNoSkin.Len() < MaterialName.Len())
-							{
-								//If we have remove the _skinXX search for material without the suffixe
-								Material = FindObject<UMaterialInterface>(ANY_PACKAGE, *MaterialNameNoSkin);
-							}
+							Material = FindObject<UMaterialInterface>(ANY_PACKAGE, *ImportedMaterial.MaterialImportName);
 						}
 					}
 					
 					const bool bEnableShadowCasting = true;
 					const bool bInRecomputeTangent = false;
-					Materials.Add(FSkeletalMaterial(Material, bEnableShadowCasting, bInRecomputeTangent, Material != nullptr ? Material->GetFName() : FName(*MaterialNameNoSkin), FName(*(ImportedMaterial.MaterialImportName))));
+					Materials.Add(FSkeletalMaterial(Material, bEnableShadowCasting, bInRecomputeTangent, Material != nullptr ? Material->GetFName() : FName(*ImportedMaterial.MaterialImportName), FName(*(ImportedMaterial.MaterialImportName))));
 				}
 
 				int32 NumMaterialsToAdd = FMath::Max<int32>(ImportedMaterials.Num(), ImportData.MaxMaterialIndex + 1);
@@ -997,7 +850,7 @@ UClass* UInterchangeSkeletalMeshFactory::GetFactoryClass() const
 
 UObject* UInterchangeSkeletalMeshFactory::CreateEmptyAsset(const FCreateAssetParams& Arguments) const
 {
-#if 1//!WITH_EDITOR || !WITH_EDITORONLY_DATA
+#if !WITH_EDITOR || !WITH_EDITORONLY_DATA
 
 	UE_LOG(LogInterchangeImport, Error, TEXT("Cannot import skeletalMesh asset in runtime, this is an editor only feature."));
 	return nullptr;
@@ -1009,8 +862,8 @@ UObject* UInterchangeSkeletalMeshFactory::CreateEmptyAsset(const FCreateAssetPar
 		return nullptr;
 	}
 
-	const UInterchangeSkeletalMeshNode* SkeletalMeshNode = Cast<UInterchangeSkeletalMeshNode>(Arguments.AssetNode);
-	if (SkeletalMeshNode == nullptr)
+	const UInterchangeSkeletalMeshFactoryNode* SkeletalMeshFactoryNode = Cast<UInterchangeSkeletalMeshFactoryNode>(Arguments.AssetNode);
+	if (SkeletalMeshFactoryNode == nullptr)
 	{
 		return nullptr;
 	}
@@ -1037,7 +890,7 @@ UObject* UInterchangeSkeletalMeshFactory::CreateEmptyAsset(const FCreateAssetPar
 	
 	SkeletalMesh->PreEditChange(nullptr);
 	//Allocate the LODImport data in the main thread
-	SkeletalMesh->ReserveLODImportData(SkeletalMeshNode->GetLodDataCount());
+	SkeletalMesh->ReserveLODImportData(SkeletalMeshFactoryNode->GetLodDataCount());
 	
 	return SkeletalMesh;
 #endif //else !WITH_EDITOR || !WITH_EDITORONLY_DATA
@@ -1045,7 +898,7 @@ UObject* UInterchangeSkeletalMeshFactory::CreateEmptyAsset(const FCreateAssetPar
 
 UObject* UInterchangeSkeletalMeshFactory::CreateAsset(const UInterchangeSkeletalMeshFactory::FCreateAssetParams& Arguments) const
 {
-#if 1//!WITH_EDITOR || !WITH_EDITORONLY_DATA
+#if !WITH_EDITOR || !WITH_EDITORONLY_DATA
 
 	UE_LOG(LogInterchangeImport, Error, TEXT("Cannot import skeletalMesh asset in runtime, this is an editor only feature."));
 	return nullptr;
@@ -1057,8 +910,8 @@ UObject* UInterchangeSkeletalMeshFactory::CreateAsset(const UInterchangeSkeletal
 		return nullptr;
 	}
 
-	const UInterchangeSkeletalMeshNode* SkeletalMeshNode = Cast<UInterchangeSkeletalMeshNode>(Arguments.AssetNode);
-	if (SkeletalMeshNode == nullptr)
+	const UInterchangeSkeletalMeshFactoryNode* SkeletalMeshFactoryNode = Cast<UInterchangeSkeletalMeshFactoryNode>(Arguments.AssetNode);
+	if (SkeletalMeshFactoryNode == nullptr)
 	{
 		return nullptr;
 	}
@@ -1070,7 +923,7 @@ UObject* UInterchangeSkeletalMeshFactory::CreateAsset(const UInterchangeSkeletal
 		return nullptr;
 	}
 
-	const UClass* SkeletalMeshClass = SkeletalMeshNode->GetAssetClass();
+	const UClass* SkeletalMeshClass = SkeletalMeshFactoryNode->GetAssetClass();
 	check(SkeletalMeshClass && SkeletalMeshClass->IsChildOf(GetFactoryClass()));
 
 	// create an asset if it doesn't exist
@@ -1118,9 +971,9 @@ UObject* UInterchangeSkeletalMeshFactory::CreateAsset(const UInterchangeSkeletal
 				ImportedResource->LODModels.Empty();
 			}
 			USkeleton* SkeletonReference = nullptr;
-			int32 LodCount = SkeletalMeshNode->GetLodDataCount();
+			int32 LodCount = SkeletalMeshFactoryNode->GetLodDataCount();
 			TArray<FString> LodDataUniqueIds;
-			SkeletalMeshNode->GetLodDataUniqueIds(LodDataUniqueIds);
+			SkeletalMeshFactoryNode->GetLodDataUniqueIds(LodDataUniqueIds);
 			ensure(LodDataUniqueIds.Num() == LodCount);
 			int32 CurrentLodIndex = 0;
 			for (int32 LodIndex = 0; LodIndex < LodCount; ++LodIndex)
@@ -1132,13 +985,66 @@ UObject* UInterchangeSkeletalMeshFactory::CreateAsset(const UInterchangeSkeletal
 					UE_LOG(LogInterchangeImport, Warning, TEXT("Invalid LOD when importing SkeletalMesh asset %s"), *Arguments.AssetName);
 					continue;
 				}
-				FString SkeletonNodeId;
-				if (!LodDataNode->GetCustomSkeletonID(SkeletonNodeId))
+
+				//Get the mesh node context for each MeshUids
+				struct FMeshNodeContext
+				{
+					const UInterchangeMeshNode* MeshNode = nullptr;
+					const UInterchangeSceneNode* SceneNode = nullptr;
+					FTransform SceneGlobalTransform = FTransform::Identity;
+					FString TranslatorPayloadKey;
+				};
+				TArray<FMeshNodeContext> MeshReferences;
+				//Scope to query the mesh node
+				{
+					TArray<FString> MeshUids;
+					LodDataNode->GetMeshUids(MeshUids);
+					MeshReferences.Reserve(MeshUids.Num());
+					for (const FString& MeshUid : MeshUids)
+					{
+						FMeshNodeContext MeshReference;
+						MeshReference.MeshNode = Cast<UInterchangeMeshNode>(Arguments.NodeContainer->GetNode(MeshUid));
+						if (!MeshReference.MeshNode)
+						{
+							//The reference is a scene node and we need to bake the geometry
+							MeshReference.SceneNode = Cast<UInterchangeSceneNode>(Arguments.NodeContainer->GetNode(MeshUid));
+							if (!ensure(MeshReference.SceneNode != nullptr))
+							{
+								UE_LOG(LogInterchangeImport, Warning, TEXT("Invalid LOD mesh reference when importing SkeletalMesh asset %s"), *Arguments.AssetName);
+								continue;
+							}
+							FString MeshDependencyUid;
+							MeshReference.SceneNode->GetCustomMeshDependencyUid(MeshDependencyUid);
+							MeshReference.MeshNode = Cast<UInterchangeMeshNode>(Arguments.NodeContainer->GetNode(MeshDependencyUid));
+							//Cache the scene node global matrix, we will use this matrix to bake the vertices
+							MeshReference.SceneNode->GetCustomGlobalTransform(MeshReference.SceneGlobalTransform);
+						}
+						if (!ensure(MeshReference.MeshNode != nullptr))
+						{
+							UE_LOG(LogInterchangeImport, Warning, TEXT("Invalid LOD mesh reference when importing SkeletalMesh asset %s"), *Arguments.AssetName);
+							continue;
+						}
+						TOptional<FString> MeshPayloadKey = MeshReference.MeshNode->GetPayLoadKey();
+						if (MeshPayloadKey.IsSet())
+						{
+							MeshReference.TranslatorPayloadKey = MeshPayloadKey.GetValue();
+						}
+						else
+						{
+							UE_LOG(LogInterchangeImport, Warning, TEXT("Empty LOD mesh reference payload when importing SkeletalMesh asset %s"), *Arguments.AssetName);
+							continue;
+						}
+						MeshReferences.Add(MeshReference);
+					}
+				}
+
+				FString SkeletonNodeUid;
+				if (!LodDataNode->GetCustomSkeletonUid(SkeletonNodeUid))
 				{
 					UE_LOG(LogInterchangeImport, Warning, TEXT("Invalid Skeleton LOD when importing SkeletalMesh asset %s"), *Arguments.AssetName);
 					continue;
 				}
-				const UInterchangeSkeletonNode* SkeletonNode = Cast<UInterchangeSkeletonNode>(Arguments.NodeContainer->GetNode(SkeletonNodeId));
+				const UInterchangeSkeletonFactoryNode* SkeletonNode = Cast<UInterchangeSkeletonFactoryNode>(Arguments.NodeContainer->GetNode(SkeletonNodeUid));
 				if (!SkeletonNode)
 				{
 					UE_LOG(LogInterchangeImport, Warning, TEXT("Invalid Skeleton LOD when importing SkeletalMesh asset %s"), *Arguments.AssetName);
@@ -1160,7 +1066,7 @@ UObject* UInterchangeSkeletalMeshFactory::CreateAsset(const UInterchangeSkeletal
 				}
 
 				FString RootJointNodeId;
-				if (!SkeletonNode->GetCustomRootJointID(RootJointNodeId))
+				if (!SkeletonNode->GetCustomRootJointUid(RootJointNodeId))
 				{
 					UE_LOG(LogInterchangeImport, Warning, TEXT("Invalid Skeleton LOD Root Joint when importing SkeletalMesh asset %s"), *Arguments.AssetName);
 					continue;
@@ -1172,10 +1078,69 @@ UObject* UInterchangeSkeletalMeshFactory::CreateAsset(const UInterchangeSkeletal
 
 				//Add the lod mesh data to the skeletalmesh
 				FSkeletalMeshImportData SkeletalMeshImportData;
-				TArray<FString> TranslatorMeshKeys;
-				LodDataNode->GetTranslatorMeshKeys(TranslatorMeshKeys);
-				//We should have only one key per LOD
-				ensure(TranslatorMeshKeys.Num() == 1);
+				FMeshDescription LodMeshDescription;
+				FSkeletalMeshAttributes SkeletalMeshAttributes(LodMeshDescription);
+				SkeletalMeshAttributes.Register();
+				FStaticMeshOperations::FAppendSettings AppendSettings;
+				for (int32 ChannelIdx = 0; ChannelIdx < FStaticMeshOperations::FAppendSettings::MAX_NUM_UV_CHANNELS; ++ChannelIdx)
+				{
+					AppendSettings.bMergeUVChannels[ChannelIdx] = true;
+				}
+				//Fill the lod mesh description using all combined mesh part
+				for (const FMeshNodeContext& MeshNodeContext : MeshReferences)
+				{
+					TOptional<UE::Interchange::FSkeletalMeshLodPayloadData> LodMeshPayload = SkeletalMeshTranslatorPayloadInterface->GetSkeletalMeshLodPayloadData(MeshNodeContext.TranslatorPayloadKey);
+					if (!LodMeshPayload.IsSet())
+					{
+						UE_LOG(LogInterchangeImport, Warning, TEXT("Invalid Skeletal mesh payload key [%s] SkeletalMesh asset %s"), *MeshNodeContext.TranslatorPayloadKey, *Arguments.AssetName);
+						continue;
+					}
+					const int32 VertexOffset = LodMeshDescription.Vertices().Num();
+					const int32 VertexInstanceOffset = LodMeshDescription.VertexInstances().Num();
+					const int32 TriangleOffset = LodMeshDescription.Triangles().Num();
+					
+					FSkeletalMeshOperations::FSkeletalMeshAppendSettings SkeletalMeshAppendSettings;
+					SkeletalMeshAppendSettings.SourceVertexIDOffset = VertexOffset;
+					FElementIDRemappings ElementIDRemappings;
+					LodMeshPayload->LodMeshDescription.Compact(ElementIDRemappings);
+					//Remap the influence vertex index to point on the correct index
+					if (LodMeshPayload->JointNames.Num() > 0)
+					{
+						const int32 LocalJointCount = LodMeshPayload->JointNames.Num();
+						const int32 RefBoneCount = RefBonesBinary.Num();
+						SkeletalMeshAppendSettings.SourceRemapBoneIndex.AddZeroed(LocalJointCount);
+						for(int32 LocalJointIndex = 0; LocalJointIndex < LocalJointCount; ++LocalJointIndex)
+						{
+							SkeletalMeshAppendSettings.SourceRemapBoneIndex[LocalJointIndex] = LocalJointIndex;
+							const FString& LocalJointName = LodMeshPayload->JointNames[LocalJointIndex];
+							for(int32 RefBoneIndex = 0; RefBoneIndex < RefBoneCount; ++RefBoneIndex)
+							{
+								const SkeletalMeshImportData::FBone& Bone = RefBonesBinary[RefBoneIndex];
+								if(Bone.Name.Equals(LocalJointName))
+								{
+									SkeletalMeshAppendSettings.SourceRemapBoneIndex[LocalJointIndex] = RefBoneIndex;
+									break;
+								}
+							}
+						}
+					}
+
+					FStaticMeshOperations::AppendMeshDescription(LodMeshPayload->LodMeshDescription, LodMeshDescription, AppendSettings);
+					FSkeletalMeshOperations::AppendSkinWeight(LodMeshPayload->LodMeshDescription, LodMeshDescription, SkeletalMeshAppendSettings);
+
+					/*
+					TArray<FString> BlendShapeToImport;
+					MeshNodeContext.MeshNode->GetShapeDependencies(BlendShapeToImport);
+					for(const FString)
+					//Copy also the blend shapes data (we need to pass the offset in the Main MeshDescription, so we have the correct vertex index mapping)
+					UE::Interchange::Private::CopyBlendShapesMeshDescriptionToSkeletalMeshImportData(BlendShapeToImport, LodMeshPayload->LodBlendShapeMeshDescriptions, SkeletalMeshImportData);
+					*/
+					
+				}
+				UE::Interchange::Private::CopyMeshDescriptionToSkeletalMeshImportData(LodMeshDescription, SkeletalMeshImportData);
+				SkeletalMeshImportData.RefBonesBinary = RefBonesBinary;
+/*
+
 				//Use the first geo LOD since we are suppose to have only one Key
 				if(TranslatorMeshKeys.IsValidIndex(0))
 				{
@@ -1201,26 +1166,27 @@ UObject* UInterchangeSkeletalMeshFactory::CreateAsset(const UInterchangeSkeletal
 					//Copy also the blend shapes data
 					UE::Interchange::Private::CopyBlendShapesMeshDescriptionToSkeletalMeshImportData(BlendShapeToImport, LodMeshPayload->LodBlendShapeMeshDescriptions, SkeletalMeshImportData);
 				}
+*/
 
 				ensure(ImportedResource->LODModels.Add(new FSkeletalMeshLODModel()) == CurrentLodIndex);
 				FSkeletalMeshLODModel& LODModel = ImportedResource->LODModels[CurrentLodIndex];
 
 				TMap<FString, UMaterialInterface*> AvailableMaterials;
-				TArray<FString> Dependencies;
-				SkeletalMeshNode->GetDependecies(Dependencies);
-				for (int32 DependencyIndex = 0; DependencyIndex < Dependencies.Num(); ++DependencyIndex)
+				TArray<FString> FactoryDependencies;
+				SkeletalMeshFactoryNode->GetFactoryDependencies(FactoryDependencies);
+				for (int32 DependencyIndex = 0; DependencyIndex < FactoryDependencies.Num(); ++DependencyIndex)
 				{
-					const UInterchangeMaterialNode* MaterialNode = Cast<UInterchangeMaterialNode>(Arguments.NodeContainer->GetNode(Dependencies[DependencyIndex]));
-					if (!MaterialNode || !MaterialNode->ReferenceObject.IsValid())
+					const UInterchangeMaterialFactoryNode* MaterialFactoryNode = Cast<UInterchangeMaterialFactoryNode>(Arguments.NodeContainer->GetNode(FactoryDependencies[DependencyIndex]));
+					if (!MaterialFactoryNode || !MaterialFactoryNode->ReferenceObject.IsValid())
 					{
 						continue;
 					}
-					UMaterialInterface* MaterialInterface = Cast<UMaterialInterface>(MaterialNode->ReferenceObject.ResolveObject());
+					UMaterialInterface* MaterialInterface = Cast<UMaterialInterface>(MaterialFactoryNode->ReferenceObject.ResolveObject());
 					if (!MaterialInterface)
 					{
 						continue;
 					}
-					AvailableMaterials.Add(MaterialNode->GetDisplayLabel(), MaterialInterface);
+					AvailableMaterials.Add(MaterialFactoryNode->GetDisplayLabel(), MaterialInterface);
 				}
 
 				UE::Interchange::Private::ProcessImportMeshMaterials(SkeletalMesh->GetMaterials(), SkeletalMeshImportData, AvailableMaterials);
@@ -1261,8 +1227,8 @@ UObject* UInterchangeSkeletalMeshFactory::CreateAsset(const UInterchangeSkeletal
 
 			SkeletalMesh->CalculateInvRefMatrices();
 
-			/** Apply all SkeletalMeshNode custom attributes to the material asset */
-			SkeletalMeshNode->ApplyAllCustomAttributeToAsset(SkeletalMesh);
+			/** Apply all SkeletalMeshFactoryNode custom attributes to the material asset */
+			SkeletalMeshFactoryNode->ApplyAllCustomAttributeToAsset(SkeletalMesh);
 		}
 		
 		//Getting the file Hash will cache it into the source data
diff --git a/Engine/Plugins/Runtime/Interchange/Source/Import/Private/Mesh/InterchangeSkeletonFactory.cpp b/Engine/Plugins/Runtime/Interchange/Source/Import/Private/Mesh/InterchangeSkeletonFactory.cpp
--- a/Engine/Plugins/Runtime/Interchange/Source/Import/Private/Mesh/InterchangeSkeletonFactory.cpp
+++ b/Engine/Plugins/Runtime/Interchange/Source/Import/Private/Mesh/InterchangeSkeletonFactory.cpp
@@ -8,8 +8,7 @@
 #include "InterchangeSourceData.h"
 #include "Nodes/InterchangeBaseNode.h"
 #include "Nodes/InterchangeBaseNodeContainer.h"
-
-
+#include "Rendering/SkeletalMeshLODImporterData.h"
 
 UClass* UInterchangeSkeletonFactory::GetFactoryClass() const
 {
@@ -35,7 +34,7 @@ UObject* UInterchangeSkeletonFactory::CreateEmptyAsset(const FCreateAssetParams&
 	// create an asset if it doesn't exist
 	UObject* ExistingAsset = StaticFindObject(nullptr, Arguments.Parent, *Arguments.AssetName);
 
-	// create a new material or overwrite existing asset, if possible
+	// create a new skeleton or overwrite existing asset, if possible
 	if (!ExistingAsset)
 	{
 		Skeleton = NewObject<UObject>(Arguments.Parent, USkeleton::StaticClass(), *Arguments.AssetName, RF_Public | RF_Standalone);
@@ -51,6 +50,8 @@ UObject* UInterchangeSkeletonFactory::CreateEmptyAsset(const FCreateAssetParams&
 		UE_LOG(LogInterchangeImport, Warning, TEXT("Could not create Skeleton asset %s"), *Arguments.AssetName);
 		return nullptr;
 	}
+	SkeletonNode->ReferenceObject = Skeleton;
+
 	Skeleton->PreEditChange(nullptr);
 #endif //WITH_EDITORONLY_DATA
 	return Skeleton;
@@ -83,7 +84,7 @@ UObject* UInterchangeSkeletonFactory::CreateAsset(const UInterchangeSkeletonFact
 	UObject* ExistingAsset = StaticFindObject(nullptr, Arguments.Parent, *Arguments.AssetName);
 
 	UObject* SkeletonObject = nullptr;
-	// create a new material or overwrite existing asset, if possible
+	// create a new skeleton or overwrite existing asset, if possible
 	if (!ExistingAsset)
 	{
 		//NewObject is not thread safe, the asset registry directory watcher tick on the main thread can trig before we finish initializing the UObject and will crash
@@ -105,8 +106,8 @@ UObject* UInterchangeSkeletonFactory::CreateAsset(const UInterchangeSkeletonFact
 
 	if (SkeletonObject)
 	{
-		//Currently material re-import will not touch the material at all
-		//TODO design a re-import process for the material (expressions and input connections)
+		//Currently skeleton re-import will not touch the skeleton at all
+		//TODO design a re-import process for the skeleton (expressions and input connections)
 		if(!Arguments.ReimportObject)
 		{
 			USkeleton* Skeleton = Cast<USkeleton>(SkeletonObject);
@@ -115,16 +116,25 @@ UObject* UInterchangeSkeletonFactory::CreateAsset(const UInterchangeSkeletonFact
 				UE_LOG(LogInterchangeImport, Warning, TEXT("Could not create Skeleton asset %s"), *Arguments.AssetName);
 				return nullptr;
 			}
+			FString RootJointUid;
+			if (!SkeletonNode->GetCustomRootJointUid(RootJointUid))
+			{
+				UE_LOG(LogInterchangeImport, Warning, TEXT("Could not create Skeleton asset %s, because there is no valid root joint node"), *Arguments.AssetName);
+				return nullptr;
+			}
+			
+			//The joint will be added by the skeletalmesh factory since we need a valid skeletalmesh to add joint to a skeleton
+			SkeletonNode->ReferenceObject = Skeleton;
 		}
 		
 		//Getting the file Hash will cache it into the source data
 		Arguments.SourceData->GetFileContentHash();
 
-		//The interchange completion task (call in the GameThread after the factories pass), will call PostEditChange which will trig another asynchronous system that will build all material in parallel
+		//The interchange completion task (call in the GameThread after the factories pass), will call PostEditChange which will trig another asynchronous system that will build all skeleton in parallel
 	}
 	else
 	{
-		//The material is not a UMaterialInterface
+		//The skeleton is not a USkeleton
 		SkeletonObject->RemoveFromRoot();
 		SkeletonObject->MarkPendingKill();
 	}
diff --git a/Engine/Plugins/Runtime/Interchange/Source/Import/Public/Mesh/InterchangeSkeletalMeshPayload.h b/Engine/Plugins/Runtime/Interchange/Source/Import/Public/Mesh/InterchangeSkeletalMeshPayload.h
--- a/Engine/Plugins/Runtime/Interchange/Source/Import/Public/Mesh/InterchangeSkeletalMeshPayload.h
+++ b/Engine/Plugins/Runtime/Interchange/Source/Import/Public/Mesh/InterchangeSkeletalMeshPayload.h
@@ -13,7 +13,9 @@ namespace UE
 		{
 			//Currently the skeletalmesh payload data is editor only, we have to move to something available at runtime
 			FMeshDescription LodMeshDescription;
-			TMap<FString, FMeshDescription> LodBlendShapeMeshDescriptions;
+
+			//This map the indice use in the meshdescription to the bone name, so we can use this information to remap properly the skinning when we merge the meshdescription
+			TArray<FString> JointNames;
 		};
 	}//ns Interchange
 }//ns UE
diff --git a/Engine/Plugins/Runtime/Interchange/Source/Nodes/Public/InterchangeMeshNode.h b/Engine/Plugins/Runtime/Interchange/Source/Nodes/Public/InterchangeMeshNode.h
--- a/Engine/Plugins/Runtime/Interchange/Source/Nodes/Public/InterchangeMeshNode.h
+++ b/Engine/Plugins/Runtime/Interchange/Source/Nodes/Public/InterchangeMeshNode.h
@@ -54,6 +54,12 @@ namespace UE
 				static FString Dependencies_BaseKey = TEXT("__MeshShapeDependencies__");
 				return Dependencies_BaseKey;
 			}
+
+			static const FString& GetSceneInstancesUidsKey()
+			{
+				static FString SceneInstanceUids_BaseKey = TEXT("__MeshSceneInstancesUids__");
+				return SceneInstanceUids_BaseKey;
+			}
 		};
 
 	}//ns Interchange
@@ -71,6 +77,7 @@ public:
 		SkeletonDependencies.Initialize(Attributes, UE::Interchange::FMeshNodeStaticData::GetSkeletonDependenciesKey());
 		MaterialDependencies.Initialize(Attributes, UE::Interchange::FMeshNodeStaticData::GetMaterialDependenciesKey());
 		ShapeDependencies.Initialize(Attributes, UE::Interchange::FMeshNodeStaticData::GetShapeDependenciesKey());
+		SceneInstancesUids.Initialize(Attributes, UE::Interchange::FMeshNodeStaticData::GetSceneInstancesUidsKey());
 	}
 
 	virtual FString GetKeyDisplayName(const UE::Interchange::FAttributeKey& NodeAttributeKey) const override
@@ -137,6 +144,21 @@ public:
 			}
 			return KeyDisplayName;
 		}
+		else if (NodeAttributeKey.Key.Equals(UE::Interchange::FMeshNodeStaticData::GetSceneInstancesUidsKey()))
+		{
+			return KeyDisplayName = TEXT("Scene mesh instances count");
+		}
+		else if (NodeAttributeKey.Key.StartsWith(UE::Interchange::FMeshNodeStaticData::GetSceneInstancesUidsKey()))
+		{
+			KeyDisplayName = TEXT("Scene mesh instances Index ");
+			const FString IndexKey = UE::Interchange::FNameAttributeArrayHelper::IndexKey();
+			int32 IndexPosition = NodeAttributeKey.Key.Find(IndexKey) + IndexKey.Len();
+			if (IndexPosition < NodeAttributeKey.Key.Len())
+			{
+				KeyDisplayName += NodeAttributeKey.Key.RightChop(IndexPosition);
+			}
+			return KeyDisplayName;
+		}
 		return Super::GetKeyDisplayName(NodeAttributeKey);
 	}
 
@@ -154,6 +176,10 @@ public:
 		{
 			return FString(TEXT("ShapeDependencies"));
 		}
+		else if (NodeAttributeKey.Key.StartsWith(UE::Interchange::FMeshNodeStaticData::GetSceneInstancesUidsKey()))
+		{
+			return FString(TEXT("SceneInstances"));
+		}
 		else if (NodeAttributeKey == Macro_CustomVertexCountKey
 				 || NodeAttributeKey == Macro_CustomPolygonCountKey
 				 || NodeAttributeKey == Macro_CustomBoundingBoxKey
@@ -601,6 +627,54 @@ public:
 	{
 		return ShapeDependencies.RemoveName(DependencyUid);
 	}
+
+	/**
+	 * This function allow to retrieve the number of scene node instancing this mesh.
+	 *
+	 */
+	UFUNCTION(BlueprintCallable, Category = "Interchange | Node | Mesh")
+	int32 GetSceneInstanceUidsCount() const
+	{
+		return SceneInstancesUids.GetCount();
+	}
+
+	/**
+	 * This function allow to retrieve the Shape dependency for this object.
+	 *
+	 */
+	UFUNCTION(BlueprintCallable, Category = "Interchange | Node | Mesh")
+	void GetSceneInstanceUids(TArray<FString>& OutDependencies) const
+	{
+		SceneInstancesUids.GetNames(OutDependencies);
+	}
+
+	/**
+	 * This function allow to retrieve one Shape dependency for this object.
+	 *
+	 */
+	UFUNCTION(BlueprintCallable, Category = "Interchange | Node | Mesh")
+	void GetSceneInstanceUid(const int32 Index, FString& OutDependency) const
+	{
+		SceneInstancesUids.GetName(Index, OutDependency);
+	}
+
+	/**
+	 * Add one Shape dependency to this object.
+	 */
+	UFUNCTION(BlueprintCallable, Category = "Interchange | Node | Mesh")
+	bool SetSceneInstanceUid(const FString& DependencyUid)
+	{
+		return SceneInstancesUids.AddName(DependencyUid);
+	}
+
+	/**
+	 * Remove one Shape dependency from this object.
+	 */
+	UFUNCTION(BlueprintCallable, Category = "Interchange | Node | Mesh")
+	bool RemoveSceneInstanceUid(const FString& DependencyUid)
+	{
+		return SceneInstancesUids.RemoveName(DependencyUid);
+	}
 private:
 	const UE::Interchange::FAttributeKey Macro_CustomVertexCountKey = UE::Interchange::FAttributeKey(TEXT("VertexCount"));
 	const UE::Interchange::FAttributeKey Macro_CustomPolygonCountKey = UE::Interchange::FAttributeKey(TEXT("PolygonCount"));
@@ -615,4 +689,5 @@ private:
 	UE::Interchange::FNameAttributeArrayHelper SkeletonDependencies;
 	UE::Interchange::FNameAttributeArrayHelper MaterialDependencies;
 	UE::Interchange::FNameAttributeArrayHelper ShapeDependencies;
+	UE::Interchange::FNameAttributeArrayHelper SceneInstancesUids;
 };
diff --git a/Engine/Plugins/Runtime/Interchange/Source/Nodes/Public/InterchangeSceneNode.h b/Engine/Plugins/Runtime/Interchange/Source/Nodes/Public/InterchangeSceneNode.h
--- a/Engine/Plugins/Runtime/Interchange/Source/Nodes/Public/InterchangeSceneNode.h
+++ b/Engine/Plugins/Runtime/Interchange/Source/Nodes/Public/InterchangeSceneNode.h
@@ -22,10 +22,22 @@ namespace UE
 				return SceneNodeSpecializeType_BaseKey;
 			}
 
-			static const FString& GetAssetDependenciesBaseKey()
+			static const FString& GetMaterialDependencyUidsBaseKey()
 			{
-				static FString AssetDependencies_BaseKey = TEXT("__AssetDependenciesBaseKey__");
-				return AssetDependencies_BaseKey;
+				static FString MaterialDependencyUids_BaseKey = TEXT("__MaterialDependencyUidsBaseKey__");
+				return MaterialDependencyUids_BaseKey;
+			}
+
+			static const FString& GetJointSpecializeTypeString()
+			{
+				static FString JointSpecializeTypeString = TEXT("Joint");
+				return JointSpecializeTypeString;
+			}
+
+			static const FString& GetLodGroupSpecializeTypeString()
+			{
+				static FString JointSpecializeTypeString = TEXT("LodGroup");
+				return JointSpecializeTypeString;
 			}
 		};
 
@@ -42,7 +54,7 @@ public:
 	:UInterchangeBaseNode()
 	{
 		NodeSpecializeTypes.Initialize(Attributes, UE::Interchange::FSceneNodeStaticData::GetNodeSpecializeTypeBaseKey());
-		AssetDependencies.Initialize(Attributes, UE::Interchange::FSceneNodeStaticData::GetAssetDependenciesBaseKey());
+		MaterialDependencyUids.Initialize(Attributes, UE::Interchange::FSceneNodeStaticData::GetMaterialDependencyUidsBaseKey());
 	}
 
 	/**
@@ -73,14 +85,14 @@ public:
 			}
 			return KeyDisplayName;
 		}
-		else if (NodeAttributeKey.Key.Equals(UE::Interchange::FSceneNodeStaticData::GetAssetDependenciesBaseKey()))
+		else if (NodeAttributeKey.Key.Equals(UE::Interchange::FSceneNodeStaticData::GetMaterialDependencyUidsBaseKey()))
 		{
-			KeyDisplayName = TEXT("Asset dependencies count");
+			KeyDisplayName = TEXT("Material dependencies count");
 			return KeyDisplayName;
 		}
-		else if (NodeAttributeKey.Key.StartsWith(UE::Interchange::FSceneNodeStaticData::GetAssetDependenciesBaseKey()))
+		else if (NodeAttributeKey.Key.StartsWith(UE::Interchange::FSceneNodeStaticData::GetMaterialDependencyUidsBaseKey()))
 		{
-			KeyDisplayName = TEXT("Asset dependency index ");
+			KeyDisplayName = TEXT("Material dependency index ");
 			const FString IndexKey = UE::Interchange::FNameAttributeArrayHelper::IndexKey();
 			int32 IndexPosition = NodeAttributeKey.Key.Find(IndexKey) + IndexKey.Len();
 			if (IndexPosition < NodeAttributeKey.Key.Len())
@@ -98,12 +110,13 @@ public:
 		{
 			return FString(TEXT("SpecializeType"));
 		}
-		else if (NodeAttributeKey.Key.StartsWith(UE::Interchange::FSceneNodeStaticData::GetAssetDependenciesBaseKey()))
+		else if (NodeAttributeKey.Key.StartsWith(UE::Interchange::FSceneNodeStaticData::GetMaterialDependencyUidsBaseKey()))
 		{
-			return FString(TEXT("AssetDependencies"));
+			return FString(TEXT("MaterialDependencies"));
 		}
 		else if (NodeAttributeKey == Macro_CustomLocalTransformKey
-				 || NodeAttributeKey == Macro_CustomGlobalTransformKey)
+				 || NodeAttributeKey == Macro_CustomGlobalTransformKey
+				 || NodeAttributeKey == Macro_CustomMeshDependencyUidKey)
 		{
 			return FString(TEXT("Scene"));
 		}
@@ -130,6 +143,22 @@ public:
 		return FName(*SpecializedType);
 	}
 
+	/** Return true if this node contains the specialized type parameter.*/
+	UFUNCTION(BlueprintCallable, Category = "Interchange | Node | Scene")
+	bool IsSpecializedTypeContains(const FString& SpecializedType) const
+	{
+		TArray<FString> SpecializedTypes;
+		GetSpecializedTypes(SpecializedTypes);
+		for (const FString& SpecializedTypeRef : SpecializedTypes)
+		{
+			if (SpecializedTypeRef.Equals(SpecializedType))
+			{
+				return true;
+			}
+		}
+		return false;
+	}
+
 	/** Specialized type are scene node special type like (Joint or LODGroup).*/
 	UFUNCTION(BlueprintCallable, Category = "Interchange | Node | Scene")
 	int32 GetSpecializedTypeCount() const
@@ -163,33 +192,33 @@ public:
 
 	/** Asset dependencies are the asset on which this node depend.*/
 	UFUNCTION(BlueprintCallable, Category = "Interchange | Node | Scene")
-	int32 GetAssetDependenciesCount() const
+	int32 GetMaterialDependencyUidsCount() const
 	{
-		return AssetDependencies.GetCount();
+		return MaterialDependencyUids.GetCount();
 	}
 
 	UFUNCTION(BlueprintCallable, Category = "Interchange | Node | Scene")
-	void GetAssetDependency(const int32 Index, FString& OutAssetDependency) const
+	void GetMaterialDependencyUid(const int32 Index, FString& OutMaterialDependencyUid) const
 	{
-		AssetDependencies.GetName(Index, OutAssetDependency);
+		MaterialDependencyUids.GetName(Index, OutMaterialDependencyUid);
 	}
 
 	UFUNCTION(BlueprintCallable, Category = "Interchange | Node | Scene")
-	void GetAssetDependencies(TArray<FString>& OutAssetDependencies) const
+	void GetMaterialDependencyUids(TArray<FString>& OutMaterialDependencyUids) const
 	{
-		AssetDependencies.GetNames(OutAssetDependencies);
+		MaterialDependencyUids.GetNames(OutMaterialDependencyUids);
 	}
 
 	UFUNCTION(BlueprintCallable, Category = "Interchange | Node | Scene")
-	bool AddAssetDependency(const FString& AssetDependency)
+	bool AddMaterialDependencyUid(const FString& MaterialDependencyUid)
 	{
-		return AssetDependencies.AddName(AssetDependency);
+		return MaterialDependencyUids.AddName(MaterialDependencyUid);
 	}
 
 	UFUNCTION(BlueprintCallable, Category = "Interchange | Node | Scene")
-	bool RemoveAssetDependency(const FString& AssetDependency)
+	bool RemoveMaterialDependencyUid(const FString& MaterialDependencyUid)
 	{
-		return AssetDependencies.RemoveName(AssetDependency);
+		return MaterialDependencyUids.RemoveName(MaterialDependencyUid);
 	}
 
 	/** Return false if the Attribute was not set previously.*/
@@ -218,11 +247,25 @@ public:
 		IMPLEMENT_NODE_ATTRIBUTE_SETTER_NODELEGATE(GlobalTransform, FTransform);
 	}
 
+	/** Return false if the Attribute was not set previously.*/
+	UFUNCTION(BlueprintCallable, Category = "Interchange | Node | Scene")
+	bool GetCustomMeshDependencyUid(FString& AttributeValue) const
+	{
+		IMPLEMENT_NODE_ATTRIBUTE_GETTER(MeshDependencyUid, FString);
+	}
+
+	UFUNCTION(BlueprintCallable, Category = "Interchange | Node | Scene")
+	bool SetCustomMeshDependencyUid(const FString& AttributeValue)
+	{
+		IMPLEMENT_NODE_ATTRIBUTE_SETTER_NODELEGATE(MeshDependencyUid, FString);
+	}
+
 private:
 	//Scene Attribute Keys
 	const UE::Interchange::FAttributeKey Macro_CustomLocalTransformKey = UE::Interchange::FAttributeKey(TEXT("LocalTransform"));
 	const UE::Interchange::FAttributeKey Macro_CustomGlobalTransformKey = UE::Interchange::FAttributeKey(TEXT("GlobalTransform"));
+	const UE::Interchange::FAttributeKey Macro_CustomMeshDependencyUidKey = UE::Interchange::FAttributeKey(TEXT("MeshDependencyUid"));
 
 	UE::Interchange::FNameAttributeArrayHelper NodeSpecializeTypes;
-	UE::Interchange::FNameAttributeArrayHelper AssetDependencies;
+	UE::Interchange::FNameAttributeArrayHelper MaterialDependencyUids;
 };
diff --git a/Engine/Plugins/Runtime/Interchange/Source/Nodes/Public/InterchangeSkeletalMeshLodDataNode.h b/Engine/Plugins/Runtime/Interchange/Source/Nodes/Public/InterchangeSkeletalMeshLodDataNode.h
--- a/Engine/Plugins/Runtime/Interchange/Source/Nodes/Public/InterchangeSkeletalMeshLodDataNode.h
+++ b/Engine/Plugins/Runtime/Interchange/Source/Nodes/Public/InterchangeSkeletalMeshLodDataNode.h
@@ -15,16 +15,10 @@ namespace UE
 
 		struct FSkeletalMeshNodeLodDataStaticData : public FBaseNodeStaticData
 		{
-			static const FString& GetBlendShapesNameBaseKey()
+			static const FString& GetMeshUidsBaseKey()
 			{
-				static FString BlendShapesName_BaseKey = TEXT("BlendShape_Names");
-				return BlendShapesName_BaseKey;
-			}
-
-			static const FString& GetTranslatorMeshKeysBaseKey()
-			{
-				static FString TranslatorMeshKeys_BaseKey = TEXT("TranslatorMesh_Keys");
-				return TranslatorMeshKeys_BaseKey;
+				static FString MeshUids_BaseKey = TEXT("__MeshUids__Key");
+				return MeshUids_BaseKey;
 			}
 		};
 
@@ -40,8 +34,7 @@ public:
 	UInterchangeSkeletalMeshLodDataNode()
 		:UInterchangeBaseNode()
 	{
-		BlendShapesNames.Initialize(Attributes, UE::Interchange::FSkeletalMeshNodeLodDataStaticData::GetBlendShapesNameBaseKey());
-		TranslatorMeshKeys.Initialize(Attributes, UE::Interchange::FSkeletalMeshNodeLodDataStaticData::GetTranslatorMeshKeysBaseKey());
+		MeshUids.Initialize(Attributes, UE::Interchange::FSkeletalMeshNodeLodDataStaticData::GetMeshUidsBaseKey());
 	}
 
 	/**
@@ -53,82 +46,92 @@ public:
 		return TypeName;
 	}
 
-	virtual FGuid GetHash() const override
-	{
-		return Attributes->GetStorageHash();
-	}
-
-
-public:
-	/** Return false if the Attribute was not set previously.*/
-	UFUNCTION(BlueprintCallable, Category = "Interchange | Node | SkeletalMeshLodData")
-	bool GetCustomSkeletonID(FString& AttributeValue) const
+	virtual FString GetKeyDisplayName(const UE::Interchange::FAttributeKey& NodeAttributeKey) const override
 	{
-		IMPLEMENT_NODE_ATTRIBUTE_GETTER(SkeletonID, FString);
+		FString KeyDisplayName = NodeAttributeKey.Key;
+		if (NodeAttributeKey.Key.Equals(UE::Interchange::FSkeletalMeshNodeLodDataStaticData::GetMeshUidsBaseKey()))
+		{
+			KeyDisplayName = TEXT("Mesh count");
+			return KeyDisplayName;
+		}
+		else if (NodeAttributeKey.Key.StartsWith(UE::Interchange::FSkeletalMeshNodeLodDataStaticData::GetMeshUidsBaseKey()))
+		{
+			KeyDisplayName = TEXT("Mesh index ");
+			const FString IndexKey = UE::Interchange::FNameAttributeArrayHelper::IndexKey();
+			int32 IndexPosition = NodeAttributeKey.Key.Find(IndexKey) + IndexKey.Len();
+			if (IndexPosition < NodeAttributeKey.Key.Len())
+			{
+				KeyDisplayName += NodeAttributeKey.Key.RightChop(IndexPosition);
+			}
+			return KeyDisplayName;
+		}
+		else if (NodeAttributeKey == Macro_CustomSkeletonUidKey)
+		{
+			KeyDisplayName = TEXT("Skeleton factory node");
+			return KeyDisplayName;
+		}
+		return Super::GetKeyDisplayName(NodeAttributeKey);
 	}
 
-	UFUNCTION(BlueprintCallable, Category = "Interchange | Node | SkeletalMeshLodData")
-	bool SetCustomSkeletonID(const FString& AttributeValue)
+	virtual FString GetAttributeCategory(const UE::Interchange::FAttributeKey& NodeAttributeKey) const override
 	{
-		IMPLEMENT_NODE_ATTRIBUTE_SETTER_NODELEGATE(SkeletonID, FString)
+		if (NodeAttributeKey.Key.StartsWith(UE::Interchange::FSkeletalMeshNodeLodDataStaticData::GetMeshUidsBaseKey()))
+		{
+			return FString(TEXT("Meshes"));
+		}
+		return Super::GetAttributeCategory(NodeAttributeKey);
 	}
 
-	/* Translator mesh keys*/
-
-	UFUNCTION(BlueprintCallable, Category = "Interchange | Node | SkeletalMeshLodData")
-	int32 GetTranslatorMeshKeyCount() const
+	virtual FGuid GetHash() const override
 	{
-		return TranslatorMeshKeys.GetCount();
+		return Attributes->GetStorageHash();
 	}
 
-	UFUNCTION(BlueprintCallable, Category = "Interchange | Node | SkeletalMeshLodData")
-	void GetTranslatorMeshKeys(TArray<FString>& OutTranslatorMeshKeys) const
-	{
-		TranslatorMeshKeys.GetNames(OutTranslatorMeshKeys);
-	}
 
+public:
+	/** Return false if the Attribute was not set previously.*/
 	UFUNCTION(BlueprintCallable, Category = "Interchange | Node | SkeletalMeshLodData")
-	bool AddTranslatorMeshKey(const FString& TranslatorMeshKey)
+	bool GetCustomSkeletonUid(FString& AttributeValue) const
 	{
-		return TranslatorMeshKeys.AddName(TranslatorMeshKey);
+		IMPLEMENT_NODE_ATTRIBUTE_GETTER(SkeletonUid, FString);
 	}
 
 	UFUNCTION(BlueprintCallable, Category = "Interchange | Node | SkeletalMeshLodData")
-	bool RemoveTranslatorMeshKey(const FString& TranslatorMeshKey)
+	bool SetCustomSkeletonUid(const FString& AttributeValue)
 	{
-		return TranslatorMeshKeys.RemoveName(TranslatorMeshKey);
+		IMPLEMENT_NODE_ATTRIBUTE_SETTER_NODELEGATE(SkeletonUid, FString)
 	}
 
-	/* Blend shapes*/
+	/* Mesh Uids: It can be either a scene or a mesh node uid. If its a scene it mean we want the mesh factory to bake the geo payload with the global transform of the scene node. */
 
 	UFUNCTION(BlueprintCallable, Category = "Interchange | Node | SkeletalMeshLodData")
-	int32 GetBlendShapeCount() const
+	int32 GetMeshUidsCount() const
 	{
-		return BlendShapesNames.GetCount();
+		return MeshUids.GetCount();
 	}
 
 	UFUNCTION(BlueprintCallable, Category = "Interchange | Node | SkeletalMeshLodData")
-	void GetBlendShapes(TArray<FString>& OutBlendShapeNames) const
+	void GetMeshUids(TArray<FString>& OutBlendShapeNames) const
 	{
-		BlendShapesNames.GetNames(OutBlendShapeNames);
+		MeshUids.GetNames(OutBlendShapeNames);
 	}
 
 	UFUNCTION(BlueprintCallable, Category = "Interchange | Node | SkeletalMeshLodData")
-	bool AddBlendShape(const FString& BlendShapeName)
+	bool AddMeshUid(const FString& BlendShapeName)
 	{
-		return BlendShapesNames.AddName(BlendShapeName);
+		return MeshUids.AddName(BlendShapeName);
 	}
 
 	UFUNCTION(BlueprintCallable, Category = "Interchange | Node | SkeletalMeshLodData")
-	bool RemoveBlendShape(const FString& BlendShapeName)
+	bool RemoveMeshUid(const FString& BlendShapeName)
 	{
-		return BlendShapesNames.RemoveName(BlendShapeName);
+		return MeshUids.RemoveName(BlendShapeName);
 	}
 
 	UFUNCTION(BlueprintCallable, Category = "Interchange | Node | SkeletalMeshLodData")
-	bool RemoveAllBlendShape()
+	bool RemoveAllMeshes()
 	{
-		return BlendShapesNames.RemoveAllNames();
+		return MeshUids.RemoveAllNames();
 	}
 
 private:
@@ -143,9 +146,8 @@ private:
 	}
 
 	//SkeletalMesh
-	const UE::Interchange::FAttributeKey Macro_CustomSkeletonIDKey = UE::Interchange::FAttributeKey(TEXT("SkeletonID"));
+	const UE::Interchange::FAttributeKey Macro_CustomSkeletonUidKey = UE::Interchange::FAttributeKey(TEXT("__SkeletonUid__Key"));
 
-	UE::Interchange::FNameAttributeArrayHelper BlendShapesNames;
-	UE::Interchange::FNameAttributeArrayHelper TranslatorMeshKeys;
+	UE::Interchange::FNameAttributeArrayHelper MeshUids;
 protected:
 };
diff --git a/Engine/Plugins/Runtime/Interchange/Source/Nodes/Public/InterchangeSkeletonFactoryNode.h b/Engine/Plugins/Runtime/Interchange/Source/Nodes/Public/InterchangeSkeletonFactoryNode.h
--- a/Engine/Plugins/Runtime/Interchange/Source/Nodes/Public/InterchangeSkeletonFactoryNode.h
+++ b/Engine/Plugins/Runtime/Interchange/Source/Nodes/Public/InterchangeSkeletonFactoryNode.h
@@ -83,13 +83,13 @@ public:
 public:
 	/** Return false if the Attribute was not set previously.*/
 	UFUNCTION(BlueprintCallable, Category = "Interchange | Node | Skeleton")
-	bool GetCustomRootJointID(FString& AttributeValue) const
+	bool GetCustomRootJointUid(FString& AttributeValue) const
 	{
 		IMPLEMENT_NODE_ATTRIBUTE_GETTER(RootJointID, FString);
 	}
 
 	UFUNCTION(BlueprintCallable, Category = "Interchange | Node | Skeleton")
-	bool SetCustomRootJointID(const FString& AttributeValue)
+	bool SetCustomRootJointUid(const FString& AttributeValue)
 	{
 		IMPLEMENT_NODE_ATTRIBUTE_SETTER_NODELEGATE(RootJointID, FString)
 	}
diff --git a/Engine/Plugins/Runtime/Interchange/Source/Parsers/Fbx/Private/FbxMaterial.cpp b/Engine/Plugins/Runtime/Interchange/Source/Parsers/Fbx/Private/FbxMaterial.cpp
--- a/Engine/Plugins/Runtime/Interchange/Source/Parsers/Fbx/Private/FbxMaterial.cpp
+++ b/Engine/Plugins/Runtime/Interchange/Source/Parsers/Fbx/Private/FbxMaterial.cpp
@@ -202,7 +202,7 @@ namespace UE
 					FbxSurfaceMaterial* SurfaceMaterial = ParentFbxNode->GetMaterial(MaterialIndex);
 					UInterchangeMaterialNode* MaterialNode = AddNodeMaterial(SurfaceMaterial, NodeContainer, JSonErrorMessages);
 					//The dependencies order is important because mesh will use index in that order to determine material use by a face
-					SceneNode->AddAssetDependency(MaterialNode->GetUniqueID());
+					SceneNode->AddMaterialDependencyUid(MaterialNode->GetUniqueID());
 				}
 			}
 
diff --git a/Engine/Plugins/Runtime/Interchange/Source/Parsers/Fbx/Private/FbxMesh.cpp b/Engine/Plugins/Runtime/Interchange/Source/Parsers/Fbx/Private/FbxMesh.cpp
--- a/Engine/Plugins/Runtime/Interchange/Source/Parsers/Fbx/Private/FbxMesh.cpp
+++ b/Engine/Plugins/Runtime/Interchange/Source/Parsers/Fbx/Private/FbxMesh.cpp
@@ -203,14 +203,14 @@ namespace UE
 				return FillMeshDescriptionFromFbxMesh(Mesh, UnusedArray, EMeshType::Static);
 			}
 
-			bool FMeshDescriptionImporter::FillSkinnedMeshDescriptionFromFbxMesh(FbxMesh* Mesh, TArray<FString>& OutJointNodeUniqueIDs)
+			bool FMeshDescriptionImporter::FillSkinnedMeshDescriptionFromFbxMesh(FbxMesh* Mesh, TArray<FString>& OutJointUniqueNames)
 			{
 				if (!ensure(bInitialized) || !ensure(MeshDescription) || !ensure(Mesh) || !ensure(Mesh->GetDeformerCount() > 0))
 				{
 					return false;
 				}
 
-				return FillMeshDescriptionFromFbxMesh(Mesh, OutJointNodeUniqueIDs, EMeshType::Skinned);
+				return FillMeshDescriptionFromFbxMesh(Mesh, OutJointUniqueNames, EMeshType::Skinned);
 			}
 
 			bool FMeshDescriptionImporter::FillMeshDescriptionFromFbxShape(FbxShape* Shape)
@@ -279,7 +279,7 @@ namespace UE
 				return true;
 			}
 
-			bool FMeshDescriptionImporter::FillMeshDescriptionFromFbxMesh(FbxMesh* Mesh, TArray<FString>& OutJointNodeUniqueIDs, EMeshType MeshType)
+			bool FMeshDescriptionImporter::FillMeshDescriptionFromFbxMesh(FbxMesh* Mesh, TArray<FString>& OutJointUniqueNames, EMeshType MeshType)
 			{
 				if (!ensure(bInitialized) || !ensure(MeshDescription) || !ensure(Mesh))
 				{
@@ -287,6 +287,7 @@ namespace UE
 				}
 
 				FStaticMeshAttributes Attributes(*MeshDescription);
+				Attributes.Register();
 
 				//Get the base layer of the mesh
 				FbxLayer* BaseLayer = Mesh->GetLayer(0);
@@ -997,7 +998,7 @@ namespace UE
 						const int32 ClusterCount = Skin->GetClusterCount();
 						TArray<FbxNode*> SortedJoints;
 						SortedJoints.Reserve(ClusterCount);
-						OutJointNodeUniqueIDs.Reserve(ClusterCount);
+						OutJointUniqueNames.Reserve(ClusterCount);
 						// create influences for each cluster
 						for (int32 ClusterIndex = 0; ClusterIndex < ClusterCount; ClusterIndex++)
 						{
@@ -1016,8 +1017,8 @@ namespace UE
 							if (!SortedJoints.Find(Link, BoneIndex))
 							{
 								BoneIndex = SortedJoints.Add(Link);
-								FString JointNodeUniqueID = FFbxHelper::GetFbxNodeHierarchyName(Link);
-								OutJointNodeUniqueIDs.Add(JointNodeUniqueID);
+								FString JointNodeUniqueID = FFbxHelper::GetFbxObjectName(Link);
+								OutJointUniqueNames.Add(JointNodeUniqueID);
 							}
 
 							//	get the vertex indices
@@ -1123,7 +1124,7 @@ namespace UE
 
 				bool bFetchSkinnedData = (Mesh->GetDeformerCount(FbxDeformer::eSkin) > 0);
 				
-				TArray<FString> JointUniqueID;
+				TArray<FString> JointUniqueNames;
 
 				FMeshDescription MeshDescription;
 				if (bFetchSkinnedData)
@@ -1131,7 +1132,7 @@ namespace UE
 					FSkeletalMeshAttributes SkeletalMeshAttribute(MeshDescription);
 					SkeletalMeshAttribute.Register();
 					FMeshDescriptionImporter MeshDescriptionImporter(&MeshDescription, SDKScene, SDKGeometryConverter);
-					if (!MeshDescriptionImporter.FillSkinnedMeshDescriptionFromFbxMesh(Mesh, JointUniqueID))
+					if (!MeshDescriptionImporter.FillSkinnedMeshDescriptionFromFbxMesh(Mesh, JointUniqueNames))
 					{
 						JSonErrorMessages.Add(TEXT("{\"Msg\" : {\"Type\" : \"Error\",\n\"Msg\" : \"Cannot fetch skinned mesh payload because there was an error when creating the FMeshDescription\"}}"));
 						return false;
@@ -1157,7 +1158,7 @@ namespace UE
 					if (bFetchSkinnedData)
 					{
 						//When passing a skinned MeshDescription, We want to pass the joint Node ID so we can know what the influence bone index refer to
-						Ar << JointUniqueID;
+						Ar << JointUniqueNames;
 					}
 					uint8* ArchiveData = Ar.GetData();
 					int64 ArchiveSize = Ar.TotalSize();
diff --git a/Engine/Plugins/Runtime/Interchange/Source/Parsers/Fbx/Private/FbxMesh.h b/Engine/Plugins/Runtime/Interchange/Source/Parsers/Fbx/Private/FbxMesh.h
--- a/Engine/Plugins/Runtime/Interchange/Source/Parsers/Fbx/Private/FbxMesh.h
+++ b/Engine/Plugins/Runtime/Interchange/Source/Parsers/Fbx/Private/FbxMesh.h
@@ -30,7 +30,7 @@ namespace UE
 				/*
 				 * Fill the mesh description using the Mesh parameter and also fill the OutJointNodeUniqueIDs so the MeshDescription bone Index can be map to the correct interchange joint scene node.
 				 */
-				bool FillSkinnedMeshDescriptionFromFbxMesh(FbxMesh* Mesh, TArray<FString>& OutJointNodeUniqueIDs);
+				bool FillSkinnedMeshDescriptionFromFbxMesh(FbxMesh* Mesh, TArray<FString>& OutJointUniqueNames);
 
 				/*
 				 * Fill the mesh description using the Shape parameter.
@@ -46,7 +46,7 @@ namespace UE
 					Skinned = 2, //skinned mesh with joints
 				};
 
-				bool FillMeshDescriptionFromFbxMesh(FbxMesh* Mesh, TArray<FString>& OutJointNodeUniqueIDs, EMeshType MeshType);
+				bool FillMeshDescriptionFromFbxMesh(FbxMesh* Mesh, TArray<FString>& OutJointUniqueNames, EMeshType MeshType);
 				bool IsOddNegativeScale(FbxAMatrix& TotalMatrix);
 				
 				//TODO move the real function from RenderCore to FVector, so we do not have to add render core to compute such a simple thing
diff --git a/Engine/Plugins/Runtime/Interchange/Source/Parsers/Fbx/Private/FbxScene.cpp b/Engine/Plugins/Runtime/Interchange/Source/Parsers/Fbx/Private/FbxScene.cpp
--- a/Engine/Plugins/Runtime/Interchange/Source/Parsers/Fbx/Private/FbxScene.cpp
+++ b/Engine/Plugins/Runtime/Interchange/Source/Parsers/Fbx/Private/FbxScene.cpp
@@ -8,6 +8,7 @@
 #include "FbxInclude.h"
 #include "FbxMaterial.h"
 #include "FbxMesh.h"
+#include "InterchangeMeshNode.h"
 #include "InterchangeSceneNode.h"
 #include "Nodes/InterchangeBaseNodeContainer.h"
 
@@ -21,14 +22,14 @@ namespace UE
 			{
 				void CreateMeshNodeReference(UInterchangeSceneNode* UnrealSceneNode, FbxNode* FbxSceneNode, FbxNodeAttribute* NodeAttribute, UInterchangeBaseNodeContainer& NodeContainer, TArray<FString>& JSonErrorMessages)
 				{
-					UInterchangeBaseNode* MeshNode = nullptr;
+					UInterchangeMeshNode* MeshNode = nullptr;
 					if (NodeAttribute->GetAttributeType() == FbxNodeAttribute::eMesh)
 					{
 						FbxMesh* Mesh = static_cast<FbxMesh*>(NodeAttribute);
 						if (ensure(Mesh))
 						{
 							FString MeshRefString = FFbxMesh::GetMeshUniqueID(Mesh);
-							MeshNode = NodeContainer.GetNode(MeshRefString);
+							MeshNode = Cast<UInterchangeMeshNode>(NodeContainer.GetNode(MeshRefString));
 						}
 					}
 					else if (NodeAttribute->GetAttributeType() == FbxNodeAttribute::eShape)
@@ -38,7 +39,8 @@ namespace UE
 
 					if (MeshNode)
 					{
-						UnrealSceneNode->AddAssetDependency(MeshNode->GetUniqueID());
+						UnrealSceneNode->SetCustomMeshDependencyUid(MeshNode->GetUniqueID());
+						MeshNode->SetSceneInstanceUid(UnrealSceneNode->GetUniqueID());
 					}
 				}
 
@@ -59,20 +61,29 @@ namespace UE
 					UInterchangeSceneNode* UnrealNode = FFbxScene::CreateTransformNode(NodeContainer, NodeName, NodeUniqueID, JSonErrorMessages);
 					check(UnrealNode);
 					
-					FTransform GlobalTransform = FFbxConvert::ConvertTransform(Node->EvaluateGlobalTransform());
-					UnrealNode->SetCustomGlobalTransform(GlobalTransform);
-					FbxNode* FbxParentNode = Node->GetParent();
-					FTransform LocalTransform;
-					if (FbxParentNode)
+					auto GetConvertedTransform = [](FbxAMatrix& NewFbxMatrix)
 					{
-						FTransform ParentGlobalTransform = FFbxConvert::ConvertTransform(FbxParentNode->EvaluateGlobalTransform());
-						LocalTransform = ParentGlobalTransform.Inverse() * GlobalTransform;
-					}
-					else
+						FTransform Transform;
+						FbxVector4 NewLocalT = NewFbxMatrix.GetT();
+						FbxVector4 NewLocalS = NewFbxMatrix.GetS();
+						FbxQuaternion NewLocalQ = NewFbxMatrix.GetQ();
+						Transform.SetTranslation(FFbxConvert::ConvertPos(NewLocalT));
+						Transform.SetScale3D(FFbxConvert::ConvertScale(NewLocalS));
+						Transform.SetRotation(FFbxConvert::ConvertRotToQuat(NewLocalQ));
+						return Transform;
+					};
+
 					{
-						LocalTransform = GlobalTransform;
+						//Set the global node transform
+						FbxAMatrix GloabalFbxMatrix = Node->EvaluateGlobalTransform();
+						FTransform GlobalTransform = GetConvertedTransform(GloabalFbxMatrix);
+						UnrealNode->SetCustomGlobalTransform(GlobalTransform);
+
+						//Set the local node transform
+						FbxAMatrix LocalFbxMatrix = Node->EvaluateLocalTransform();
+						FTransform LocalTransform = GetConvertedTransform(LocalFbxMatrix);
+						UnrealNode->SetCustomLocalTransform(LocalTransform);
 					}
-					UnrealNode->SetCustomLocalTransform(LocalTransform);
 
 					int32 AttributeCount = Node->GetNodeAttributeCount();
 					for (int32 AttributeIndex = 0; AttributeIndex < AttributeCount; ++AttributeIndex)
@@ -103,8 +114,7 @@ namespace UE
 							case FbxNodeAttribute::eSkeleton:
 							{
 								//Add the joint specialized type
-								FString SpecializedType = TEXT("Joint");
-								UnrealNode->AddSpecializedType(SpecializedType);
+								UnrealNode->AddSpecializedType(FSceneNodeStaticData::GetJointSpecializeTypeString());
 								break;
 							}
 							case FbxNodeAttribute::eMesh:
@@ -116,8 +126,7 @@ namespace UE
 							}
 							case FbxNodeAttribute::eLODGroup:
 							{
-								FString SpecializedType = TEXT("LodGroup");
-								UnrealNode->AddSpecializedType(SpecializedType);
+								UnrealNode->AddSpecializedType(FSceneNodeStaticData::GetLodGroupSpecializeTypeString());
 								break;
 							}
 							case FbxNodeAttribute::eCamera:
diff --git a/Engine/Plugins/Runtime/Interchange/Source/Pipelines/Private/InterchangeGenericAssetsPipeline.cpp b/Engine/Plugins/Runtime/Interchange/Source/Pipelines/Private/InterchangeGenericAssetsPipeline.cpp
--- a/Engine/Plugins/Runtime/Interchange/Source/Pipelines/Private/InterchangeGenericAssetsPipeline.cpp
+++ b/Engine/Plugins/Runtime/Interchange/Source/Pipelines/Private/InterchangeGenericAssetsPipeline.cpp
@@ -1,12 +1,16 @@
 // Copyright Epic Games, Inc. All Rights Reserved. 
 #include "InterchangeGenericAssetsPipeline.h"
 
+#include "Animation/Skeleton.h"
 #include "CoreMinimal.h"
 #include "InterchangeMaterialFactoryNode.h"
 #include "InterchangeMaterialNode.h"
 #include "InterchangeMeshNode.h"
 #include "InterchangePipelineLog.h"
 #include "InterchangeSceneNode.h"
+#include "InterchangeSkeletalMeshFactoryNode.h"
+#include "InterchangeSkeletalMeshLodDataNode.h"
+#include "InterchangeSkeletonFactoryNode.h"
 #include "InterchangeTextureFactoryNode.h"
 #include "InterchangeTextureNode.h"
 #include "InterchangeSourceData.h"
@@ -77,14 +81,66 @@ bool UInterchangeGenericAssetsPipeline::ExecutePreImportPipeline(UInterchangeBas
 	//Find all skinned mesh
 	//import meshes
 	TArray<UInterchangeMeshNode*> SkinnedMeshNodes;
+	TArray<UInterchangeMeshNode*> StaticMeshNodes;
 	for (const TObjectPtr<UInterchangeMeshNode> MeshNode : MeshNodes)
 	{
 		if (MeshNode->IsSkinnedMesh())
 		{
 			SkinnedMeshNodes.Add(MeshNode.Get());
 		}
+		else
+		{
+			StaticMeshNodes.Add(MeshNode.Get());
+		}
 	}
 
+	if (bCombineSkeletalMeshes && SkinnedMeshNodes.Num() > 1)
+	{
+		//Get All the root joint nodes
+		TArray<FString> SkeletonRootNodeUids;
+		for (const TObjectPtr<UInterchangeSceneNode> SceneNode : SceneNodes)
+		{
+			if (SceneNode->IsSpecializedTypeContains(UE::Interchange::FSceneNodeStaticData::GetJointSpecializeTypeString()))
+			{
+				const UInterchangeSceneNode* ParentJointNode = Cast<UInterchangeSceneNode>(BaseNodeContainer->GetNode(SceneNode->GetParentUid()));
+				if (!ParentJointNode || !ParentJointNode->IsSpecializedTypeContains(UE::Interchange::FSceneNodeStaticData::GetJointSpecializeTypeString()))
+				{
+					SkeletonRootNodeUids.Add(SceneNode->GetUniqueID());
+				}
+			}
+		}
+		
+		TMap<FString, TArray<FString>> MesheUidsPerJointRootUid;
+		for (const UInterchangeMeshNode* SkinnedMeshNode : SkinnedMeshNodes)
+		{
+			if (SkinnedMeshNode->GetSkeletonDependeciesCount() == 0)
+			{
+				continue;
+			}
+			FString JointNodeUid;
+			SkinnedMeshNode->GetSkeletonDependency(0, JointNodeUid);
+			while (!SkeletonRootNodeUids.Contains(JointNodeUid))
+			{
+				JointNodeUid = BaseNodeContainer->GetNode(JointNodeUid)->GetParentUid();
+			}
+			if (SkeletonRootNodeUids.Contains(JointNodeUid))
+			{
+				TArray<FString>& CombinedMeshNodes = MesheUidsPerJointRootUid.FindOrAdd(JointNodeUid);
+				CombinedMeshNodes.Add(SkinnedMeshNode->GetUniqueID());
+			}
+		}
+
+		const bool bLocalUseSourceNameForAsset = bUseSourceNameForAsset && (SourceDatas.Num() == 1) && MesheUidsPerJointRootUid.Num() == 1 && SkinnedMeshNodes.Num() == MeshNodes.Num();
+		for (TPair<FString, TArray<FString>> JointRootUidAndMeshes : MesheUidsPerJointRootUid)
+		{
+			SkeletonFactoryNodes.Add(CreateSkeletonFactoryNode(JointRootUidAndMeshes.Key));
+			SkeletalMeshFactoryNodes.Add(CreateSkeletalMeshFactoryNode(JointRootUidAndMeshes.Key, JointRootUidAndMeshes.Value, bLocalUseSourceNameForAsset));
+		}
+	}
+	else
+	{
+
+	}
 	return true;
 }
 
@@ -160,9 +216,250 @@ TObjectPtr<UInterchangeMaterialFactoryNode> UInterchangeGenericAssetsPipeline::C
 	return MaterialFactoryNode;
 }
 
-TObjectPtr<UInterchangeBaseNode> UInterchangeGenericAssetsPipeline::CreateMeshFactoryNode(const TObjectPtr<UInterchangeMeshNode> MeshNode)
+TObjectPtr<UInterchangeSkeletonFactoryNode> UInterchangeGenericAssetsPipeline::CreateSkeletonFactoryNode(const FString& RootJointUid)
+{
+	const UInterchangeBaseNode* RootJointNode = BaseNodeContainer->GetNode(RootJointUid);
+	if (!RootJointNode)
+	{
+		return nullptr;
+	}
+	FString DisplayLabel = RootJointNode->GetDisplayLabel() + TEXT("_Skeleton");
+	FString SkeletonUid = TEXT("\\Skeleton\\") + RootJointNode->GetUniqueID();
+
+	UInterchangeSkeletonFactoryNode* SkeletonFactoryNode = nullptr;
+	if (BaseNodeContainer->IsNodeUidValid(SkeletonUid))
+	{
+		//The node already exist, just return it
+		SkeletonFactoryNode = Cast<UInterchangeSkeletonFactoryNode>(BaseNodeContainer->GetNode(SkeletonUid));
+		if (!ensure(SkeletonFactoryNode))
+		{
+			//Log an error
+			return nullptr;
+		}
+	}
+	else
+	{
+		SkeletonFactoryNode = NewObject<UInterchangeSkeletonFactoryNode>(BaseNodeContainer.Get(), NAME_None);
+		if (!ensure(SkeletonFactoryNode))
+		{
+			return nullptr;
+		}
+		SkeletonFactoryNode->InitializeSkeletonNode(SkeletonUid, DisplayLabel, USkeleton::StaticClass()->GetName());
+		SkeletonFactoryNode->SetCustomRootJointUid(RootJointNode->GetUniqueID());
+		BaseNodeContainer->AddNode(SkeletonFactoryNode);
+	}
+	return SkeletonFactoryNode;
+}
+
+TObjectPtr<UInterchangeSkeletalMeshFactoryNode> UInterchangeGenericAssetsPipeline::CreateSkeletalMeshFactoryNode(const FString& RootJointUid, TArray<FString>& MeshNodeUids, bool bLocalUseSourceNameForAsset)
+{
+	//Get the skeleton factory node
+	const UInterchangeBaseNode* RootJointNode = BaseNodeContainer->GetNode(RootJointUid);
+	if (!RootJointNode)
+	{
+		return nullptr;
+	}
+	const FString SkeletonUid = TEXT("\\Skeleton\\") + RootJointNode->GetUniqueID();
+	const UInterchangeSkeletonFactoryNode* SkeletonFactoryNode = Cast<UInterchangeSkeletonFactoryNode>(BaseNodeContainer->GetNode(SkeletonUid));
+	if (!ensure(SkeletonFactoryNode))
+	{
+		//Log an error
+		return nullptr;
+	}
+	
+	if (MeshNodeUids.Num() == 0)
+	{
+		return nullptr;
+	}
+	const UInterchangeMeshNode* FirstMeshNode = Cast<UInterchangeMeshNode>(BaseNodeContainer->GetNode(MeshNodeUids[0]));
+	if (!FirstMeshNode)
+	{
+		return nullptr;
+	}
+
+	//Create the skeletal mesh factory node, name it according to the first mesh node compositing the meshes
+	FString DisplayLabel = FirstMeshNode->GetDisplayLabel();
+	if (bLocalUseSourceNameForAsset && SourceDatas.Num() > 0)
+	{
+		DisplayLabel = FPaths::GetBaseFilename(SourceDatas[0]->GetFilename());
+	}
+	const FString SkeletalMeshUid = TEXT("\\SkeletalMesh\\") + RootJointNode->GetUniqueID();
+
+	UInterchangeSkeletalMeshFactoryNode* SkeletalMeshFactoryNode = NewObject<UInterchangeSkeletalMeshFactoryNode>(BaseNodeContainer.Get(), NAME_None);
+	if (!ensure(SkeletalMeshFactoryNode))
+	{
+		return nullptr;
+	}
+	SkeletalMeshFactoryNode->InitializeSkeletalMeshNode(SkeletalMeshUid, DisplayLabel, USkeletalMesh::StaticClass()->GetName());
+	SkeletalMeshFactoryNode->SetFactoryDependencyUid(SkeletonUid);
+	BaseNodeContainer->AddNode(SkeletalMeshFactoryNode);
+
+	TMap< FString, TMap<int32, TArray<FString>>> SceneInstanceUidsPerLodIndexPerMeshNodeUid;
+	for (const FString& MeshNodeUid : MeshNodeUids)
+	{
+		const UInterchangeMeshNode* MeshNode = Cast<UInterchangeMeshNode>(BaseNodeContainer->GetNode(MeshNodeUid));
+		if (!ensure(MeshNode))
+		{
+			continue;
+		}
+		//Get all scene nodes that instance this mesh node with the LOD index in case its under a LodGroup
+		TMap<int32, TArray<FString>>& SceneInstanceUidsPerLodIndex = SceneInstanceUidsPerLodIndexPerMeshNodeUid.FindOrAdd(MeshNodeUid);
+		{
+			TArray<FString> MeshSceneNodeInstanceUids;
+			MeshNode->GetSceneInstanceUids(MeshSceneNodeInstanceUids);
+			for (const FString& SceneNodeUid : MeshSceneNodeInstanceUids)
+			{
+				if (UInterchangeSceneNode* MeshSceneNode = Cast<UInterchangeSceneNode>(BaseNodeContainer->GetNode(SceneNodeUid)))
+				{
+					UInterchangeSceneNode* ParentMeshSceneNode = Cast<UInterchangeSceneNode>(BaseNodeContainer->GetNode(MeshSceneNode->GetParentUid()));
+					if (!ParentMeshSceneNode)
+					{
+						continue;
+					}
+					int32 SceneNodeLodIndex = 0;
+					FString LastChildUid = MeshSceneNode->GetUniqueID();
+					do
+					{
+						if (MeshSceneNode->IsSpecializedTypeContains(UE::Interchange::FSceneNodeStaticData::GetLodGroupSpecializeTypeString()))
+						{
+							TArray<FString> LodGroupChildrens = BaseNodeContainer->GetNodeChildrenUids(ParentMeshSceneNode->GetUniqueID());
+							for (int32 LodIndex = 0; LodIndex < LodGroupChildrens.Num(); ++LodIndex)
+							{
+								const FString& ChildrenUid = LodGroupChildrens[LodIndex];
+								if (ChildrenUid.Equals(LastChildUid))
+								{
+									SceneNodeLodIndex = LodIndex;
+									break;
+								}
+							}
+							break;
+						}
+						LastChildUid = ParentMeshSceneNode->GetUniqueID();
+						ParentMeshSceneNode = Cast<UInterchangeSceneNode>(BaseNodeContainer->GetNode(ParentMeshSceneNode->GetParentUid()));
+					} while (ParentMeshSceneNode);
+
+					//Apply bImportLod option, it will not add LOD
+					if (bImportLods || SceneNodeLodIndex == 0)
+					{
+						TArray<FString>& SceneInstanceUids = SceneInstanceUidsPerLodIndex.FindOrAdd(SceneNodeLodIndex);
+						SceneInstanceUids.AddUnique(SceneNodeUid);
+					}
+				}
+			}
+		}
+	}
+
+	for (const TPair< FString, TMap<int32, TArray<FString>>>& MeshNodeUidAndSceneInstanceUidsPerLodIndex : SceneInstanceUidsPerLodIndexPerMeshNodeUid)
+	{
+		const FString& MeshNodeUid = MeshNodeUidAndSceneInstanceUidsPerLodIndex.Key;
+		const UInterchangeMeshNode* MeshNode = Cast<UInterchangeMeshNode>(BaseNodeContainer->GetNode(MeshNodeUid));
+		if (!ensure(MeshNode))
+		{
+			continue;
+		}
+
+		const TMap<int32, TArray<FString>>& SceneInstanceUidsPerLodIndex = MeshNodeUidAndSceneInstanceUidsPerLodIndex.Value;
+		if (bBakeMeshes && SceneInstanceUidsPerLodIndex.Num() > 0)
+		{
+			//Baked and instanced meshes code path
+			AddLodDataToSkeletalMesh(SkeletonFactoryNode, SkeletalMeshFactoryNode, SceneInstanceUidsPerLodIndex);
+		}
+		else
+		{
+			//Unbaked or un-instanced meshes code path
+			AddLodDataToSkeletalMesh(SkeletonFactoryNode, SkeletalMeshFactoryNode, MeshNode);
+		}
+	}
+
+	return SkeletalMeshFactoryNode;
+}
+
+TObjectPtr<UInterchangeSkeletalMeshLodDataNode> UInterchangeGenericAssetsPipeline::CreateSkeletalMeshLodDataNode(const FString& NodeName, const FString& NodeUniqueID)
+{
+	FString DisplayLabel(NodeName);
+	FString NodeUID(NodeUniqueID);
+	UInterchangeSkeletalMeshLodDataNode* SkeletalMeshLodDataNode = NewObject<UInterchangeSkeletalMeshLodDataNode>(BaseNodeContainer.Get(), NAME_None);
+	if (!ensure(SkeletalMeshLodDataNode))
+	{
+		//TODO log error
+		return nullptr;
+	}
+	// Creating a UMaterialInterface
+	SkeletalMeshLodDataNode->InitializeNode(NodeUID, DisplayLabel, EInterchangeNodeContainerType::NodeContainerType_FactoryData);
+	BaseNodeContainer->AddNode(SkeletalMeshLodDataNode);
+	return SkeletalMeshLodDataNode;
+}
+
+void UInterchangeGenericAssetsPipeline::AddLodDataToSkeletalMesh(const UInterchangeSkeletonFactoryNode* SkeletonFactoryNode, UInterchangeSkeletalMeshFactoryNode* SkeletalMeshFactoryNode, const TMap<int32, TArray<FString>>& SceneInstanceUidsPerLodIndex)
+{
+	const FString SkeletalMeshUid = SkeletalMeshFactoryNode->GetUniqueID();
+	const FString SkeletonUid = SkeletonFactoryNode->GetUniqueID();
+	FString RootjointNodeUid;
+	SkeletonFactoryNode->GetCustomRootJointUid(RootjointNodeUid);
+	for (const TPair<int32, TArray<FString>>& LodIndexAndSceneInstanceUids : SceneInstanceUidsPerLodIndex)
+	{
+		const int32 LodIndex = LodIndexAndSceneInstanceUids.Key;
+		const TArray<FString>& SceneNodeUids = LodIndexAndSceneInstanceUids.Value;
+
+		//Create a lod data node with all the meshes for this LOD
+		const FString SkeletalMeshLodDataName = TEXT("LodData") + FString::FromInt(LodIndex);
+		const FString LODDataPrefix = TEXT("\\LodData") + (LodIndex > 0 ? FString::FromInt(LodIndex) : TEXT("")) + TEXT("\\");
+		const FString SkeletalMeshLodDataUniqueID = LODDataPrefix + RootjointNodeUid;
+		//The LodData already exist
+		UInterchangeSkeletalMeshLodDataNode* LodDataNode = Cast<UInterchangeSkeletalMeshLodDataNode>(BaseNodeContainer->GetNode(SkeletalMeshLodDataUniqueID));
+		if (!LodDataNode)
+		{
+			//Add the data for the LOD (skeleton Unique ID and all the mesh node fbx path, so we can find them when we will create the payload data)
+			LodDataNode = CreateSkeletalMeshLodDataNode(SkeletalMeshLodDataName, SkeletalMeshLodDataUniqueID);
+			LodDataNode->SetParentUid(SkeletalMeshUid);
+			LodDataNode->SetCustomSkeletonUid(SkeletonUid);
+			SkeletalMeshFactoryNode->AddLodDataUniqueId(SkeletalMeshLodDataUniqueID);
+		}
+		for (const FString& SceneNodeUid : SceneNodeUids)
+		{
+			const UInterchangeSceneNode* SceneNode = Cast<UInterchangeSceneNode>(BaseNodeContainer->GetNode(SceneNodeUid));
+			TArray<FString> MaterialDependencies;
+			SceneNode->GetMaterialDependencyUids(MaterialDependencies);
+			const int32 MaterialCount = MaterialDependencies.Num();
+			for (int32 MaterialIndex = 0; MaterialIndex < MaterialCount; ++MaterialIndex)
+			{
+				const FString MaterialFactoryNodeUid = UInterchangeMaterialFactoryNode::GetMaterialFactoryNodeUidFromMaterialNodeUid(MaterialDependencies[MaterialIndex]);
+				//Create a factory dependency so Material asset are import before the skeletal mesh asset
+				SkeletalMeshFactoryNode->SetFactoryDependencyUid(MaterialFactoryNodeUid);
+			}
+			//Add the scene node has a reference to bake all mesh instance
+			LodDataNode->AddMeshUid(SceneNodeUid);
+		}
+	}
+}
+
+void UInterchangeGenericAssetsPipeline::AddLodDataToSkeletalMesh(const UInterchangeSkeletonFactoryNode* SkeletonFactoryNode, UInterchangeSkeletalMeshFactoryNode* SkeletalMeshFactoryNode, const UInterchangeMeshNode* MeshNode)
 {
-	FString DisplayLabel = MeshNode->GetDisplayLabel();
-	FString NodeUid = MeshNode->GetUniqueID();
-	return nullptr;
+	const FString SkeletalMeshUid = SkeletalMeshFactoryNode->GetUniqueID();
+	const FString SkeletonUid = SkeletonFactoryNode->GetUniqueID();
+	FString RootjointNodeUid;
+	SkeletonFactoryNode->GetCustomRootJointUid(RootjointNodeUid);
+	//simply export the mesh node to lod 0
+	//Create a lod data node with all the meshes for this LOD
+	const FString SkeletalMeshLodDataName = TEXT("LodData0");
+	const FString SkeletalMeshLodDataUniqueID = TEXT("\\LodData\\") + RootjointNodeUid;
+	UInterchangeSkeletalMeshLodDataNode* LodDataNode = Cast<UInterchangeSkeletalMeshLodDataNode>(BaseNodeContainer->GetNode(SkeletalMeshLodDataUniqueID));
+	if (!LodDataNode)
+	{
+		//Add the data for the LOD (skeleton Unique ID and all the mesh node fbx path, so we can find them when we will create the payload data)
+		LodDataNode = CreateSkeletalMeshLodDataNode(SkeletalMeshLodDataName, SkeletalMeshLodDataUniqueID);
+		LodDataNode->SetParentUid(SkeletalMeshUid);
+		LodDataNode->SetCustomSkeletonUid(SkeletonUid);
+		SkeletalMeshFactoryNode->AddLodDataUniqueId(SkeletalMeshLodDataUniqueID);
+	}
+	TArray<FString> MaterialDependencies;
+	MeshNode->GetMaterialDependencies(MaterialDependencies);
+	const int32 MaterialCount = MaterialDependencies.Num();
+	for (int32 MaterialIndex = 0; MaterialIndex < MaterialCount; ++MaterialIndex)
+	{
+		const FString MaterialFactoryNodeUid = UInterchangeMaterialFactoryNode::GetMaterialFactoryNodeUidFromMaterialNodeUid(MaterialDependencies[MaterialIndex]);
+		//Create a factory dependency so Material asset are import before the skeletal mesh asset
+		SkeletalMeshFactoryNode->SetFactoryDependencyUid(MaterialFactoryNodeUid);
+	}
+	LodDataNode->AddMeshUid(MeshNode->GetUniqueID());
 }
\ No newline at end of file
diff --git a/Engine/Plugins/Runtime/Interchange/Source/Pipelines/Public/InterchangeGenericAssetsPipeline.h b/Engine/Plugins/Runtime/Interchange/Source/Pipelines/Public/InterchangeGenericAssetsPipeline.h
--- a/Engine/Plugins/Runtime/Interchange/Source/Pipelines/Public/InterchangeGenericAssetsPipeline.h
+++ b/Engine/Plugins/Runtime/Interchange/Source/Pipelines/Public/InterchangeGenericAssetsPipeline.h
@@ -17,6 +17,31 @@ class UInterchangeMaterialNode;
 class UInterchangeMaterialFactoryNode;
 class UInterchangeMeshNode;
 class UInterchangeSceneNode;
+class UInterchangeSkeletalMeshFactoryNode;
+class UInterchangeSkeletalMeshLodDataNode;
+class UInterchangeSkeletonFactoryNode;
+
+#define COMMON_CATEGORY "Common"
+#define COMMON_MESHES_CATEGORY "Common Meshes"
+#define STATIC_MESHES_CATEGORY "Static Meshes"
+#define SKELETAL_MESHES_CATEGORY "Skeletal Meshes"
+#define ANIMATIONS_CATEGORY "Animations"
+#define MATERIALS_CATEGORY "Materials"
+#define TEXTURES_CATEGORY "Textures"
+
+/** Force mesh type, if user want to import all meshes as one type*/
+UENUM(BlueprintType)
+enum EForceMeshType
+{
+	/** Will import from the source type, no conversion */
+	FMT_None UMETA(DisplayName = "None"),
+	/** Will import any mesh to static mesh. */
+	FMT_StaticMesh UMETA(DisplayName = "Static Mesh"),
+	/** Will import any mesh to skeletal mesh. */
+	FMT_SkeletalMesh UMETA(DisplayName = "Skeletal Mesh"),
+
+	FMT_MAX,
+};
 
 /**
  * This pipeline is the generic pipeline option for all meshes type and should be call before specialized Mesh pipeline (like generic static mesh or skeletal mesh pipelines)
@@ -30,11 +55,80 @@ class INTERCHANGEPIPELINES_API UInterchangeGenericAssetsPipeline : public UInter
 
 public:
 	//////////////////////////////////////////////////////////////////////////
-	// BEGIN Pre import pipeline properties
+	// BEGIN Pre import pipeline properties, please keep by per category order for properties declaration
+
+
+	//////	COMMON_CATEGORY Properties //////
+
+
+	/** If enable and there is only one asset and one source data, we will name the asset like the source data name. */
+	UPROPERTY(EditAnywhere, BlueprintReadWrite, Category = COMMON_CATEGORY)
+	bool bUseSourceNameForAsset = true;
+
+
+	//////	COMMON_MESHES_CATEGORY Properties //////
+
+
+	/** Allow to convert mesh to a particular type */
+	UPROPERTY(EditAnywhere, BlueprintReadWrite, Category = COMMON_MESHES_CATEGORY)
+	TEnumAsByte<enum EForceMeshType> ForceAllMeshHasType = EForceMeshType::FMT_None;
 
-	/** If enable and there is only one mesh in the translated nodes, the Source data name (filename) will be use to name the asset */
-	//UPROPERTY(EditAnywhere, BlueprintReadWrite, Category = "Miscellaneous")
-	//bool bUseSourceDataToNameMeshes = true;
+	/** If enable, meshes LODs will be imported. Note that it required the advanced bBakeMesh property to be enabled. */
+	UPROPERTY(EditAnywhere, BlueprintReadWrite, Category = COMMON_MESHES_CATEGORY, meta = (editcondition = "bBakeMeshes"))
+	bool bImportLods = true;
+
+	/** If enable, meshes will be baked with the hierarchy transform, if there is multiple instances, the mesh will incorporate all instance. */
+	UPROPERTY(EditAnywhere, BlueprintReadWrite, AdvancedDisplay, Category = COMMON_MESHES_CATEGORY)
+	bool bBakeMeshes = true;
+
+	
+	//////	STATIC_MESHES_CATEGORY Properties //////
+
+
+	/** If enable, import the animation asset find in the sources. */
+	UPROPERTY(EditAnywhere, BlueprintReadWrite, Category = STATIC_MESHES_CATEGORY)
+	bool bImportStaticMeshes = true;
+
+	/** If enable all translated static mesh node will be imported has a one static mesh. */
+	UPROPERTY(EditAnywhere, BlueprintReadWrite, Category = STATIC_MESHES_CATEGORY)
+	bool bCombineStaticMeshes = false;
+
+	
+	//////	SKELETAL_MESHES_CATEGORY Properties //////
+
+	
+	/** If enable, import the animation asset find in the sources. */
+	UPROPERTY(EditAnywhere, BlueprintReadWrite, Category = SKELETAL_MESHES_CATEGORY)
+	bool bImportSkeletalMeshes = true;
+
+	/** If enable all translated skinned mesh node will be imported has a one skeletal mesh, note that it can still create several skeletal mesh for each different root joint. */
+	UPROPERTY(EditAnywhere, BlueprintReadWrite, Category = SKELETAL_MESHES_CATEGORY)
+	bool bCombineSkeletalMeshes = true;
+
+	
+	//////	ANIMATIONS_CATEGORY Properties //////
+
+
+	/** If enable, import the animation asset find in the sources. */
+	UPROPERTY(EditAnywhere, BlueprintReadWrite, Category = ANIMATIONS_CATEGORY)
+	bool bImportAnimations = true;
+
+
+	//////	MATERIALS_CATEGORY Properties //////
+
+
+	/** If enable, import the material asset find in the sources. */
+	UPROPERTY(EditAnywhere, BlueprintReadWrite, Category = MATERIALS_CATEGORY)
+	bool bImportMaterials = true;
+
+
+	//////	TEXTURES_CATEGORY Properties //////
+
+
+	/** If enable, import the material asset find in the sources. */
+	UPROPERTY(EditAnywhere, BlueprintReadWrite, Category = TEXTURES_CATEGORY)
+	bool bImportTextures = true;
+	
 
 	// END Pre import pipeline properties
 	//////////////////////////////////////////////////////////////////////////
@@ -79,16 +173,31 @@ private:
 	/** Mesh translated assets nodes */
 	TArray<TObjectPtr<UInterchangeMeshNode>> MeshNodes;
 	
+	/** Skeleton factory assets nodes */
+	TArray<TObjectPtr<UInterchangeSkeletonFactoryNode>> SkeletonFactoryNodes;
+
+	/** Create a UInterchangeSkeletonFactorynode */
+	TObjectPtr<UInterchangeSkeletonFactoryNode> CreateSkeletonFactoryNode(const FString& RootJointUid);
+	
 	/** Skeletal mesh factory assets nodes */
-	//TArray<TObjectPtr<UInterchangeSkeletalMeshFactoryNode>> SkeletalMeshFactoryNodes;
+	TArray<TObjectPtr<UInterchangeSkeletalMeshFactoryNode>> SkeletalMeshFactoryNodes;
 	
 	/** Static mesh factory assets nodes */
 	//TArray<TObjectPtr<UInterchangeStaticMeshFactoryNode>> StaticMeshFactoryNodes;
 	
-	/** This function can create a UInterchangeSkeletalMeshFactoryNode or a UInterchangeStaticMeshFactoryNode */
-	TObjectPtr<UInterchangeBaseNode> CreateMeshFactoryNode(const TObjectPtr<UInterchangeMeshNode> MeshNode);
+	/** This function can create a UInterchangeSkeletalMeshFactoryNode */
+	TObjectPtr<UInterchangeSkeletalMeshFactoryNode> CreateSkeletalMeshFactoryNode(const FString& RootJointUid, TArray<FString>& MeshNodeUids, const bool bUseSourceNameForAsset);
+	
+	/** This function can create a UInterchangeSkeletalMeshLodDataNode which represent the LOD data need by the factory to create a lod mesh */
+	TObjectPtr<UInterchangeSkeletalMeshLodDataNode> CreateSkeletalMeshLodDataNode(const FString& NodeName, const FString& NodeUniqueID);
+
+	/** This function add all lod data node to the skeletal mesh, it handle instancing/combine and baking pipeline settings */
+	void AddLodDataToSkeletalMesh(const UInterchangeSkeletonFactoryNode* SkeletonFactoryNode, UInterchangeSkeletalMeshFactoryNode* SkeletalMeshFactoryNode, const TMap<int32, TArray<FString>>& SceneInstanceUidsPerLodIndex);
+
+	/** This function add one lod data node (the base lod) to the skeletal mesh. It is use when we do not bake the data, like for a scene import, or if a mesh is not reference by a scene node. */
+	void AddLodDataToSkeletalMesh(const UInterchangeSkeletonFactoryNode* SkeletonFactoryNode, UInterchangeSkeletalMeshFactoryNode* SkeletalMeshFactoryNode, const UInterchangeMeshNode* MeshNode);
 
-	/** Mesh translated scene nodes */
+	/** Translated scene nodes */
 	TArray<TObjectPtr<UInterchangeSceneNode>> SceneNodes;
 };
 
diff --git a/Engine/Source/Runtime/SkeletalMeshDescription/Private/SkeletalMeshOperations.cpp b/Engine/Source/Runtime/SkeletalMeshDescription/Private/SkeletalMeshOperations.cpp
--- a/Engine/Source/Runtime/SkeletalMeshDescription/Private/SkeletalMeshOperations.cpp
+++ b/Engine/Source/Runtime/SkeletalMeshDescription/Private/SkeletalMeshOperations.cpp
@@ -3,6 +3,7 @@
 #include "SkeletalMeshOperations.h"
 
 #include "SkeletalMeshAttributes.h"
+#include "BoneWeights.h"
 
 
 DEFINE_LOG_CATEGORY(LogSkeletalMeshOperations);
@@ -10,5 +11,39 @@ DEFINE_LOG_CATEGORY(LogSkeletalMeshOperations);
 #define LOCTEXT_NAMESPACE "SkeletalMeshOperations"
 
 //Add specific skeletal mesh descriptions implementation here
+void FSkeletalMeshOperations::AppendSkinWeight(const FMeshDescription& SourceMesh, FMeshDescription& TargetMesh, FSkeletalMeshAppendSettings& AppendSettings)
+{
+	TRACE_CPUPROFILER_EVENT_SCOPE_STR("FSkeletalMeshOperations::AppendSkinWeight");
+	FSkeletalMeshConstAttributes SourceSkeletalMeshAttributes(SourceMesh);
+	
+	FSkeletalMeshAttributes TargetSkeletalMeshAttributes(TargetMesh);
+	TargetSkeletalMeshAttributes.Register();
+	
+	FSkinWeightsVertexAttributesConstRef SourceVertexSkinWeights = SourceSkeletalMeshAttributes.GetVertexSkinWeights();
+	FSkinWeightsVertexAttributesRef TargetVertexSkinWeights = TargetSkeletalMeshAttributes.GetVertexSkinWeights();
+
+	TargetMesh.SuspendVertexIndexing();
+	
+	for (const FVertexID& SourceVertexID : SourceMesh.Vertices().GetElementIDs())
+	{
+		const FVertexID TargetVertexID = FVertexID(AppendSettings.SourceVertexIDOffset + SourceVertexID.GetValue());
+		FVertexBoneWeightsConst SourceBoneWeights = SourceVertexSkinWeights.Get(SourceVertexID);
+		TArray<UE::AnimationCore::FBoneWeight> TargetBoneWeights;
+		const int32 InfluenceCount = SourceBoneWeights.Num();
+		for (int32 InfluenceIndex = 0; InfluenceIndex < InfluenceCount; ++InfluenceIndex)
+		{
+			const FBoneIndexType SourceBoneIndex = SourceBoneWeights[InfluenceIndex].GetBoneIndex();
+			if(AppendSettings.SourceRemapBoneIndex.IsValidIndex(SourceBoneIndex))
+			{
+				UE::AnimationCore::FBoneWeight& TargetBoneWeight = TargetBoneWeights.AddDefaulted_GetRef();
+				TargetBoneWeight.SetBoneIndex(AppendSettings.SourceRemapBoneIndex[SourceBoneIndex]);
+				TargetBoneWeight.SetRawWeight(SourceBoneWeights[InfluenceIndex].GetRawWeight());
+			}
+		}
+		TargetVertexSkinWeights.Set(TargetVertexID, TargetBoneWeights);
+	}
+
+	TargetMesh.ResumeVertexIndexing();
+}
 
 #undef LOCTEXT_NAMESPACE
\ No newline at end of file
diff --git a/Engine/Source/Runtime/SkeletalMeshDescription/Public/SkeletalMeshOperations.h b/Engine/Source/Runtime/SkeletalMeshDescription/Public/SkeletalMeshOperations.h
--- a/Engine/Source/Runtime/SkeletalMeshDescription/Public/SkeletalMeshOperations.h
+++ b/Engine/Source/Runtime/SkeletalMeshDescription/Public/SkeletalMeshOperations.h
@@ -3,6 +3,8 @@
 #pragma once
 
 #include "CoreMinimal.h"
+
+#include "BoneIndices.h"
 #include "MeshTypes.h"
 #include "StaticMeshOperations.h"
 
@@ -14,5 +16,15 @@ DECLARE_LOG_CATEGORY_EXTERN(LogSkeletalMeshOperations, Log, All);
 class SKELETALMESHDESCRIPTION_API FSkeletalMeshOperations : public FStaticMeshOperations
 {
 public:
-	//Add any specific skeletal mesh operation here
+	struct FSkeletalMeshAppendSettings
+	{
+		FSkeletalMeshAppendSettings()
+			: SourceVertexIDOffset(0)
+		{}
+
+		int32 SourceVertexIDOffset;
+		TArray<FBoneIndexType> SourceRemapBoneIndex;
+	};
+	
+	static void AppendSkinWeight(const FMeshDescription& SourceMesh, FMeshDescription& TargetMesh, FSkeletalMeshAppendSettings& AppendSettings);
 };
