Add generic pipeline options: "BakeMeshes", "ImportStaticMeshes", "CombineStaticMeshes", "ImportSkeletalMeshes", "CombineSkeletalMeshes".
Create a pipeline meshes utilities that allow the user to query all meshes (static, skeletal, both) and also query combined meshes (static and skeletal)

diff --git a/Engine/Plugins/Runtime/Interchange/Source/Import/Private/Fbx/InterchangeFbxTranslator.cpp b/Engine/Plugins/Runtime/Interchange/Source/Import/Private/Fbx/InterchangeFbxTranslator.cpp
--- a/Engine/Plugins/Runtime/Interchange/Source/Import/Private/Fbx/InterchangeFbxTranslator.cpp
+++ b/Engine/Plugins/Runtime/Interchange/Source/Import/Private/Fbx/InterchangeFbxTranslator.cpp
@@ -214,6 +214,61 @@ TOptional<UE::Interchange::FSkeletalMeshLodPayloadData> UInterchangeFbxTranslato
 	return SkeletalMeshLodPayload;
 }
 
+TOptional<UE::Interchange::FSkeletalMeshBlendShapePayloadData> UInterchangeFbxTranslator::GetSkeletalMeshBlendShapePayloadData(const FString& PayLoadKey) const
+{
+	if (!Dispatcher.IsValid())
+	{
+		return TOptional<UE::Interchange::FSkeletalMeshBlendShapePayloadData>();
+	}
+
+	//Create a json command to read the fbx file
+	FString JsonCommand = CreateFetchPayloadFbxCommand(PayLoadKey);
+	int32 TaskIndex = Dispatcher->AddTask(JsonCommand);
+
+	//Blocking call until all tasks are executed
+	Dispatcher->WaitAllTaskToCompleteExecution();
+
+	UE::Interchange::ETaskState TaskState;
+	FString JsonResult;
+	TArray<FString> JSonMessages;
+	Dispatcher->GetTaskState(TaskIndex, TaskState, JsonResult, JSonMessages);
+
+	//TODO: Parse the JSonMessage and add the message to the interchange not yet develop error messaging
+
+	if (TaskState != UE::Interchange::ETaskState::ProcessOk)
+	{
+		return TOptional<UE::Interchange::FSkeletalMeshBlendShapePayloadData>();
+	}
+	//Grab the result file and fill the BaseNodeContainer
+	UE::Interchange::FJsonFetchPayloadCmd::JsonResultParser ResultParser;
+	ResultParser.FromJson(JsonResult);
+	FString SkeletalMeshPayloadFilename = ResultParser.GetResultFilename();
+
+	if (!ensure(FPaths::FileExists(SkeletalMeshPayloadFilename)))
+	{
+		//TODO log an error saying the payload file do not exist even if the get payload command succeed
+		return TOptional<UE::Interchange::FSkeletalMeshBlendShapePayloadData>();
+	}
+	UE::Interchange::FSkeletalMeshBlendShapePayloadData SkeletalMeshBlendShapePayload;
+	SkeletalMeshBlendShapePayload.LodMeshDescription.Empty();
+	//All sub object should be gone with the reset
+	TArray64<uint8> Buffer;
+	FFileHelper::LoadFileToArray(Buffer, *SkeletalMeshPayloadFilename);
+	uint8* FileData = Buffer.GetData();
+	int64 FileDataSize = Buffer.Num();
+	if (FileDataSize < 1)
+	{
+		//Nothing to load from this file
+		return TOptional<UE::Interchange::FSkeletalMeshBlendShapePayloadData>();
+	}
+
+	//Buffer keep the ownership of the data, the large memory reader is use to serialize the TMap
+	FLargeMemoryReader Ar(FileData, FileDataSize);
+	SkeletalMeshBlendShapePayload.LodMeshDescription.Serialize(Ar);
+
+	return SkeletalMeshBlendShapePayload;
+}
+
 FString UInterchangeFbxTranslator::CreateLoadFbxFileCommand(const FString& FbxFilePath) const
 {
 	UE::Interchange::FJsonLoadSourceCmd LoadSourceCommand(TEXT("FBX"), FbxFilePath);
diff --git a/Engine/Plugins/Runtime/Interchange/Source/Import/Private/Mesh/InterchangeSkeletalMeshFactory.cpp b/Engine/Plugins/Runtime/Interchange/Source/Import/Private/Mesh/InterchangeSkeletalMeshFactory.cpp
--- a/Engine/Plugins/Runtime/Interchange/Source/Import/Private/Mesh/InterchangeSkeletalMeshFactory.cpp
+++ b/Engine/Plugins/Runtime/Interchange/Source/Import/Private/Mesh/InterchangeSkeletalMeshFactory.cpp
@@ -136,21 +136,23 @@ namespace UE
 			}
 
 
-			bool CopyBlendShapesMeshDescriptionToSkeletalMeshImportData(const TArray<FString>& BlendShapeToImport, const TMap<FString, FMeshDescription>& LodBlendShapeMeshDescriptions, FSkeletalMeshImportData& DestinationSkeletalMeshImportData)
+			bool CopyBlendShapesMeshDescriptionToSkeletalMeshImportData(const TMap<FString, TOptional<UE::Interchange::FSkeletalMeshBlendShapePayloadData>>& LodBlendShapeMeshDescriptions, FSkeletalMeshImportData& DestinationSkeletalMeshImportData)
 			{
-				for (const TPair<FString, FMeshDescription>& Pair : LodBlendShapeMeshDescriptions)
+				for (const TPair<FString, TOptional<UE::Interchange::FSkeletalMeshBlendShapePayloadData>>& Pair : LodBlendShapeMeshDescriptions)
 				{
-					FString BlendShapeName(Pair.Key);
-					//Skip blend shape that is not define in the BlendShapeToImport array
-					if (!BlendShapeToImport.Contains(BlendShapeName))
+					const FString BlendShapeName(Pair.Key);
+					const TOptional<UE::Interchange::FSkeletalMeshBlendShapePayloadData>& BlendShapePayloadData = Pair.Value;
+					if (!BlendShapePayloadData.IsSet())
 					{
+						UE_LOG(LogInterchangeImport, Error, TEXT("Empty blend shape optional payload data [%s]"), *BlendShapeName);
 						continue;
 					}
 
+					FTransform GlobalTransform = BlendShapePayloadData->GlobalTransform.IsSet() ? BlendShapePayloadData->GlobalTransform.GetValue() : FTransform::Identity;
 					TArray<FVector3f> CompressPoints;
 					CompressPoints.Reserve(DestinationSkeletalMeshImportData.Points.Num());
 
-					const FMeshDescription& SourceMeshDescription = Pair.Value;
+					const FMeshDescription& SourceMeshDescription = BlendShapePayloadData.GetValue().LodMeshDescription;
 					FStaticMeshConstAttributes Attributes(SourceMeshDescription);
 					TVertexAttributesConstRef<FVector3f> VertexPositions = Attributes.GetVertexPositions();
 
@@ -166,7 +168,7 @@ namespace UE
 					for (FVertexID VertexID : SourceMeshDescription.Vertices().GetElementIDs())
 					{
 						//We can use GetValue because the Meshdescription was compacted before the copy
-						MorphTargetData.Points[VertexID.GetValue()] = VertexPositions[VertexID];
+						MorphTargetData.Points[VertexID.GetValue()] = GlobalTransform.TransformPosition(VertexPositions[VertexID]);
 					}
 
 					for (int32 PointIdx = 0; PointIdx < DestinationSkeletalMeshImportData.Points.Num(); ++PointIdx)
@@ -561,7 +563,7 @@ UObject* UInterchangeSkeletalMeshFactory::CreateAsset(const UInterchangeSkeletal
 				{
 					const UInterchangeMeshNode* MeshNode = nullptr;
 					const UInterchangeSceneNode* SceneNode = nullptr;
-					FTransform SceneGlobalTransform = FTransform::Identity;
+					TOptional<FTransform> SceneGlobalTransform;
 					FString TranslatorPayloadKey;
 				};
 				TArray<FMeshNodeContext> MeshReferences;
@@ -587,7 +589,11 @@ UObject* UInterchangeSkeletalMeshFactory::CreateAsset(const UInterchangeSkeletal
 							MeshReference.SceneNode->GetCustomMeshDependencyUid(MeshDependencyUid);
 							MeshReference.MeshNode = Cast<UInterchangeMeshNode>(Arguments.NodeContainer->GetNode(MeshDependencyUid));
 							//Cache the scene node global matrix, we will use this matrix to bake the vertices
-							MeshReference.SceneNode->GetCustomGlobalTransform(MeshReference.SceneGlobalTransform);
+							FTransform SceneNodeGlobalTransform;
+							if (MeshReference.SceneNode->GetCustomGlobalTransform(SceneNodeGlobalTransform))
+							{
+								MeshReference.SceneGlobalTransform = SceneNodeGlobalTransform;
+							}
 						}
 						if (!ensure(MeshReference.MeshNode != nullptr))
 						{
@@ -656,6 +662,10 @@ UObject* UInterchangeSkeletalMeshFactory::CreateAsset(const UInterchangeSkeletal
 				{
 					AppendSettings.bMergeUVChannels[ChannelIdx] = true;
 				}
+				
+				//Blend shape MeshDescriptions
+				TMap<FString, TOptional<UE::Interchange::FSkeletalMeshBlendShapePayloadData>> BlendShapeMeshDescriptionsPerBlendShapeName;
+
 				//Fill the lod mesh description using all combined mesh part
 				for (const FMeshNodeContext& MeshNodeContext : MeshReferences)
 				{
@@ -666,8 +676,6 @@ UObject* UInterchangeSkeletalMeshFactory::CreateAsset(const UInterchangeSkeletal
 						continue;
 					}
 					const int32 VertexOffset = LodMeshDescription.Vertices().Num();
-					const int32 VertexInstanceOffset = LodMeshDescription.VertexInstances().Num();
-					const int32 TriangleOffset = LodMeshDescription.Triangles().Num();
 					
 					FSkeletalMeshOperations::FSkeletalMeshAppendSettings SkeletalMeshAppendSettings;
 					SkeletalMeshAppendSettings.SourceVertexIDOffset = VertexOffset;
@@ -695,49 +703,58 @@ UObject* UInterchangeSkeletalMeshFactory::CreateAsset(const UInterchangeSkeletal
 						}
 					}
 
+					//Bake the payload, with the provide transform
+					if (MeshNodeContext.SceneGlobalTransform.IsSet())
+					{
+						AppendSettings.MeshTransform = MeshNodeContext.SceneGlobalTransform;
+					}
+					else
+					{
+						AppendSettings.MeshTransform.Reset();
+					}
 					FStaticMeshOperations::AppendMeshDescription(LodMeshPayload->LodMeshDescription, LodMeshDescription, AppendSettings);
 					FSkeletalMeshOperations::AppendSkinWeight(LodMeshPayload->LodMeshDescription, LodMeshDescription, SkeletalMeshAppendSettings);
 
-					/*
-					TArray<FString> BlendShapeToImport;
-					MeshNodeContext.MeshNode->GetShapeDependencies(BlendShapeToImport);
-					for(const FString)
-					//Copy also the blend shapes data (we need to pass the offset in the Main MeshDescription, so we have the correct vertex index mapping)
-					UE::Interchange::Private::CopyBlendShapesMeshDescriptionToSkeletalMeshImportData(BlendShapeToImport, LodMeshPayload->LodBlendShapeMeshDescriptions, SkeletalMeshImportData);
-					*/
+
 					
+					TArray<FString> BlendShapeUids;
+					MeshNodeContext.MeshNode->GetShapeDependencies(BlendShapeUids);
+					for (const FString& BlendShapeUid : BlendShapeUids)
+					{
+						if (const UInterchangeMeshNode* BlendShapeMeshNode = Cast<UInterchangeMeshNode>(Arguments.NodeContainer->GetNode(BlendShapeUid)))
+						{
+							TOptional<FString> BlendShapePayloadKey = BlendShapeMeshNode->GetPayLoadKey();
+							if (!BlendShapePayloadKey.IsSet())
+							{
+								UE_LOG(LogInterchangeImport, Warning, TEXT("Empty LOD morph target mesh reference payload when importing SkeletalMesh asset %s"), *Arguments.AssetName);
+								continue;
+							}
+							TOptional<UE::Interchange::FSkeletalMeshBlendShapePayloadData>& BlendShapeMeshPayload = BlendShapeMeshDescriptionsPerBlendShapeName.Add(BlendShapeUid);
+							BlendShapeMeshPayload = SkeletalMeshTranslatorPayloadInterface->GetSkeletalMeshBlendShapePayloadData(BlendShapePayloadKey.GetValue());
+							if (!BlendShapeMeshPayload.IsSet())
+							{
+								BlendShapeMeshDescriptionsPerBlendShapeName.Remove(BlendShapeUid);
+								UE_LOG(LogInterchangeImport, Warning, TEXT("Invalid Skeletal mesh morph target payload key [%s] SkeletalMesh asset %s"), *BlendShapePayloadKey.GetValue(), *Arguments.AssetName);
+								continue;
+							}
+							//Use the Mesh node parent bake transform
+							if (MeshNodeContext.SceneGlobalTransform.IsSet())
+							{
+								BlendShapeMeshPayload->GlobalTransform = MeshNodeContext.SceneGlobalTransform;
+							}
+							else
+							{
+								BlendShapeMeshPayload->GlobalTransform.Reset();
+							}
+						}
+					}
 				}
 
 				SkeletalMeshImportData = FSkeletalMeshImportData::CreateFromMeshDescription(LodMeshDescription);
 				SkeletalMeshImportData.RefBonesBinary = RefBonesBinary;
-/*
-
-				//Use the first geo LOD since we are suppose to have only one Key
-				if(TranslatorMeshKeys.IsValidIndex(0))
-				{
-					int32 MeshLodIndex = 0;
 
-					TOptional<UE::Interchange::FSkeletalMeshLodPayloadData> LodMeshPayload = SkeletalMeshTranslatorPayloadInterface->GetSkeletalMeshLodPayloadData(TranslatorMeshKeys[MeshLodIndex]);
-					if (!LodMeshPayload.IsSet())
-					{
-						UE_LOG(LogInterchangeImport, Warning, TEXT("Invalid Skeletal mesh payload key [%s] SkeletalMesh asset %s"), *TranslatorMeshKeys[MeshLodIndex], *Arguments.AssetName);
-						break;
-					}
-					//TODO use the mesh description as the source import data for skeletalmesh
-					//Convert the meshdescription to skeletalmesh import data and set the data for the LOD, so the build will work
-					//Make sure all IDs are compact so we can use them directly when copying the data
-					FElementIDRemappings ElementIDRemappings;
-					LodMeshPayload->LodMeshDescription.Compact(ElementIDRemappings);
-					UE::Interchange::Private::CopyMeshDescriptionToSkeletalMeshImportData(LodMeshPayload->LodMeshDescription, SkeletalMeshImportData);
-					SkeletalMeshImportData.RefBonesBinary = RefBonesBinary;
-
-					//Pipeline can remove names from the list to control which blendshapes they import
-					TArray<FString> BlendShapeToImport;
-					LodDataNode->GetBlendShapes(BlendShapeToImport);
-					//Copy also the blend shapes data
-					UE::Interchange::Private::CopyBlendShapesMeshDescriptionToSkeletalMeshImportData(BlendShapeToImport, LodMeshPayload->LodBlendShapeMeshDescriptions, SkeletalMeshImportData);
-				}
-*/
+				//Copy all the lod blend shapes data to the SkeletalMeshImportData.
+				UE::Interchange::Private::CopyBlendShapesMeshDescriptionToSkeletalMeshImportData(BlendShapeMeshDescriptionsPerBlendShapeName, SkeletalMeshImportData);
 
 				ensure(ImportedResource->LODModels.Add(new FSkeletalMeshLODModel()) == CurrentLodIndex);
 				FSkeletalMeshLODModel& LODModel = ImportedResource->LODModels[CurrentLodIndex];
diff --git a/Engine/Plugins/Runtime/Interchange/Source/Import/Public/Fbx/InterchangeFbxTranslator.h b/Engine/Plugins/Runtime/Interchange/Source/Import/Public/Fbx/InterchangeFbxTranslator.h
--- a/Engine/Plugins/Runtime/Interchange/Source/Import/Public/Fbx/InterchangeFbxTranslator.h
+++ b/Engine/Plugins/Runtime/Interchange/Source/Import/Public/Fbx/InterchangeFbxTranslator.h
@@ -80,6 +80,15 @@ public:
 	 */
 	virtual TOptional<UE::Interchange::FSkeletalMeshLodPayloadData> GetSkeletalMeshLodPayloadData(const FString& PayLoadKey) const override;
 
+	/**
+	 * Once the translation is done, the import process need a way to retrieve payload data.
+	 * This payload will be use by the factories to create the morph target shapes.
+	 *
+	 * @param PayloadKey - The key to retrieve the a particular payload contain into the specified source data.
+	 * @return a PayloadData containing the data point by the payload key. The TOptional will not be set if there is an error.
+	 */
+	virtual TOptional<UE::Interchange::FSkeletalMeshBlendShapePayloadData> GetSkeletalMeshBlendShapePayloadData(const FString& PayLoadKey) const override;
+
 	/* IInterchangeSkeletalMeshPayloadInterface End */
 private:
 
diff --git a/Engine/Plugins/Runtime/Interchange/Source/Import/Public/Mesh/InterchangeSkeletalMeshPayload.h b/Engine/Plugins/Runtime/Interchange/Source/Import/Public/Mesh/InterchangeSkeletalMeshPayload.h
--- a/Engine/Plugins/Runtime/Interchange/Source/Import/Public/Mesh/InterchangeSkeletalMeshPayload.h
+++ b/Engine/Plugins/Runtime/Interchange/Source/Import/Public/Mesh/InterchangeSkeletalMeshPayload.h
@@ -17,5 +17,12 @@ namespace UE
 			//This map the indice use in the meshdescription to the bone name, so we can use this information to remap properly the skinning when we merge the meshdescription
 			TArray<FString> JointNames;
 		};
+
+		struct FSkeletalMeshBlendShapePayloadData
+		{
+			//BlendShape payload is a MeshDescription containing at least the vertex position.
+			FMeshDescription LodMeshDescription;
+			TOptional<FTransform> GlobalTransform;
+		};
 	}//ns Interchange
 }//ns UE
diff --git a/Engine/Plugins/Runtime/Interchange/Source/Import/Public/Mesh/InterchangeSkeletalMeshPayloadInterface.h b/Engine/Plugins/Runtime/Interchange/Source/Import/Public/Mesh/InterchangeSkeletalMeshPayloadInterface.h
--- a/Engine/Plugins/Runtime/Interchange/Source/Import/Public/Mesh/InterchangeSkeletalMeshPayloadInterface.h
+++ b/Engine/Plugins/Runtime/Interchange/Source/Import/Public/Mesh/InterchangeSkeletalMeshPayloadInterface.h
@@ -31,7 +31,16 @@ public:
 	 * @param PayloadKey - The key to retrieve the a particular payload contain into the specified source data.
 	 * @return a PayloadData containing the data point by the payload key. The TOptional will not be set if there is an error.
 	 */
-	virtual TOptional<UE::Interchange::FSkeletalMeshLodPayloadData> GetSkeletalMeshLodPayloadData(const FString& PayloadKey) const = 0;
+	virtual TOptional<UE::Interchange::FSkeletalMeshLodPayloadData> GetSkeletalMeshLodPayloadData(const FString& PayLoadKey) const = 0;
+
+	/**
+	 * Once the translation is done, the import process need a way to retrieve payload data.
+	 * This payload will be use by the factories to create the morph target shapes.
+	 *
+	 * @param PayloadKey - The key to retrieve the a particular payload contain into the specified source data.
+	 * @return a PayloadData containing the data point by the payload key. The TOptional will not be set if there is an error.
+	 */
+	virtual TOptional<UE::Interchange::FSkeletalMeshBlendShapePayloadData> GetSkeletalMeshBlendShapePayloadData(const FString& PayLoadKey) const = 0;
 };
 
 
diff --git a/Engine/Plugins/Runtime/Interchange/Source/Pipelines/Private/InterchangeGenericAssetsPipeline.cpp b/Engine/Plugins/Runtime/Interchange/Source/Pipelines/Private/InterchangeGenericAssetsPipeline.cpp
--- a/Engine/Plugins/Runtime/Interchange/Source/Pipelines/Private/InterchangeGenericAssetsPipeline.cpp
+++ b/Engine/Plugins/Runtime/Interchange/Source/Pipelines/Private/InterchangeGenericAssetsPipeline.cpp
@@ -7,6 +7,7 @@
 #include "InterchangeMaterialNode.h"
 #include "InterchangeMeshNode.h"
 #include "InterchangePipelineLog.h"
+#include "InterchangePipelineMeshesUtilities.h"
 #include "InterchangeSceneNode.h"
 #include "InterchangeSkeletalMeshFactoryNode.h"
 #include "InterchangeSkeletalMeshLodDataNode.h"
@@ -15,6 +16,7 @@
 #include "InterchangeTextureNode.h"
 #include "InterchangeSourceData.h"
 #include "Misc/Paths.h"
+#include "Nodes/InterchangeBaseNode.h"
 #include "Nodes/InterchangeBaseNodeContainer.h"
 #include "UObject/Object.h"
 #include "UObject/ObjectMacros.h"
@@ -34,9 +36,13 @@ bool UInterchangeGenericAssetsPipeline::ExecutePreImportPipeline(UInterchangeBas
 		SourceDatas.Add(SourceData);
 	}
 
+	PipelineMeshesUtilities = UInterchangePipelineMeshesUtilities::CreateInterchangePipelineMeshesUtilities(BaseNodeContainer);
+	TArray<UInterchangeMeshNode*> SkinnedMeshNodes;
+	TArray<UInterchangeMeshNode*> StaticMeshNodes;
+	TMap<FString, TArray<FString>> SkeletalMeshFactoryDependencyOrderPerSkeletonRootNodeUid;
 	
 	//Find all translated node we need for this pipeline
-	BaseNodeContainer->IterateNodes([this](const FString& NodeUid, UInterchangeBaseNode* Node)
+	BaseNodeContainer->IterateNodes([this, &SkinnedMeshNodes, &StaticMeshNodes](const FString& NodeUid, UInterchangeBaseNode* Node)
 	{
 		switch(Node->GetnodeContainerType())
 		{
@@ -53,6 +59,15 @@ bool UInterchangeGenericAssetsPipeline::ExecutePreImportPipeline(UInterchangeBas
 				else if (UInterchangeMeshNode* MeshNode = Cast<UInterchangeMeshNode>(Node))
 				{
 					MeshNodes.Add(MeshNode);
+					if (MeshNode->IsSkinnedMesh())
+					{
+						SkinnedMeshNodes.Add(MeshNode);
+					}
+					else
+					{
+						StaticMeshNodes.Add(MeshNode);
+					}
+
 				}
 			}
 			break;
@@ -68,83 +83,203 @@ bool UInterchangeGenericAssetsPipeline::ExecutePreImportPipeline(UInterchangeBas
 	});
 	
 	//import textures
-	for (const TObjectPtr<UInterchangeTextureNode> TextureNode : TextureNodes)
+	for (const UInterchangeTextureNode* TextureNode : TextureNodes)
 	{
 		CreateTexture2DFactoryNode(TextureNode);
 	}
 	//import materials
-	for (const TObjectPtr<UInterchangeMaterialNode> MaterialNode : MaterialNodes)
+	for (const UInterchangeMaterialNode* MaterialNode : MaterialNodes)
 	{
 		CreateMaterialFactoryNode(MaterialNode);
 	}
 
-	//Find all skinned mesh
-	//import meshes
-	TArray<UInterchangeMeshNode*> SkinnedMeshNodes;
-	TArray<UInterchangeMeshNode*> StaticMeshNodes;
-	for (const TObjectPtr<UInterchangeMeshNode> MeshNode : MeshNodes)
+
+	if (SkinnedMeshNodes.Num() > 0)
 	{
-		if (MeshNode->IsSkinnedMesh())
-		{
-			SkinnedMeshNodes.Add(MeshNode.Get());
-		}
-		else
+		auto SetSkeletalMeshDependencies = [&SkeletalMeshFactoryDependencyOrderPerSkeletonRootNodeUid](const FString& JointNodeUid, UInterchangeSkeletalMeshFactoryNode* SkeletalMeshFactoryNode)
 		{
-			StaticMeshNodes.Add(MeshNode.Get());
-		}
-	}
+			TArray<FString>& SkeletalMeshFactoryDependencyOrder = SkeletalMeshFactoryDependencyOrderPerSkeletonRootNodeUid.FindOrAdd(JointNodeUid);
+			//Updating the skeleton is not multi thread safe, so we add dependency between skeletalmesh altering the same skeleton
+			//TODO make the skeletalMesh ReferenceSkeleton thread safe to allow multiple parallel skeletalmesh factory on the same skeleton asset.
+			int32 DependencyIndex = SkeletalMeshFactoryDependencyOrder.AddUnique(SkeletalMeshFactoryNode->GetUniqueID());
+			if (DependencyIndex > 0)
+			{
+				const FString SkeletalMeshFactoryNodeDependencyUid = SkeletalMeshFactoryDependencyOrder[DependencyIndex - 1];
+				SkeletalMeshFactoryNode->SetFactoryDependencyUid(SkeletalMeshFactoryNodeDependencyUid);
+			}
+		};
 
-	if (bCombineSkeletalMeshes && SkinnedMeshNodes.Num() > 1)
-	{
-		//Get All the root joint nodes
-		TArray<FString> SkeletonRootNodeUids;
-		for (const TObjectPtr<UInterchangeSceneNode> SceneNode : SceneNodes)
+		if (bCombineSkeletalMeshes)
 		{
-			if (SceneNode->IsSpecializedTypeContains(UE::Interchange::FSceneNodeStaticData::GetJointSpecializeTypeString()))
+			//////////////////////////////////////////////////////////////////////////
+			//Combined everything we can
+			TMap<FString, TArray<FString>> MeshUidsPerSkeletonRootUid;
+			auto CreatePerSkeletonRootUidCombinedSkinnedMesh = [this, &MeshUidsPerSkeletonRootUid, &SetSkeletalMeshDependencies](const bool bUseInstanceMesh)
 			{
-				const UInterchangeSceneNode* ParentJointNode = Cast<UInterchangeSceneNode>(BaseNodeContainer->GetNode(SceneNode->GetParentUid()));
-				if (!ParentJointNode || !ParentJointNode->IsSpecializedTypeContains(UE::Interchange::FSceneNodeStaticData::GetJointSpecializeTypeString()))
+				bool bFoundInstances = false;
+				for (const TPair<FString, TArray<FString>>& SkeletonRootUidAndMeshUids : MeshUidsPerSkeletonRootUid)
 				{
-					SkeletonRootNodeUids.Add(SceneNode->GetUniqueID());
+					const FString& SkeletonRootUid = SkeletonRootUidAndMeshUids.Key;
+					//Every iteration is a skeletalmesh asset that combine all MeshInstances sharing the same skeleton root node
+					UInterchangeSkeletonFactoryNode* SkeletonFactoryNode = CreateSkeletonFactoryNode(SkeletonRootUid);
+					//The MeshUids can represent a SceneNode pointing on a MeshNode or directly a MeshNode;
+					TMap<int32, TArray<FString>> MeshUidsPerLodIndex;
+					const TArray<FString>& MeshUids = SkeletonRootUidAndMeshUids.Value;
+					for (const FString& MeshUid : MeshUids)
+					{
+						if (bUseInstanceMesh)
+						{
+							const FInterchangeMeshInstance& MeshInstance = PipelineMeshesUtilities->GetMeshInstanceByUid(MeshUid);
+							for (const TPair<int32, FInterchangeLodSceneNodeContainer>& LodIndexAndSceneNodeContainer : MeshInstance.SceneNodePerLodIndex)
+							{
+								const int32 LodIndex = LodIndexAndSceneNodeContainer.Key;
+								const FInterchangeLodSceneNodeContainer& SceneNodeContainer = LodIndexAndSceneNodeContainer.Value;
+								TArray<FString>& TranslatedNodes = MeshUidsPerLodIndex.FindOrAdd(LodIndex);
+								for (const UInterchangeSceneNode* SceneNode : SceneNodeContainer.SceneNodes)
+								{
+									TranslatedNodes.Add(SceneNode->GetUniqueID());
+								}
+							}
+						}
+						else
+						{
+							//MeshGeometry cannot have Lod since LODs are define in the scene node
+							const FInterchangeMeshGeometry& MeshGeometry = PipelineMeshesUtilities->GetMeshGeometryByUid(MeshUid);
+							const int32 LodIndex = 0;
+							TArray<FString>& TranslatedNodes = MeshUidsPerLodIndex.FindOrAdd(LodIndex);
+							TranslatedNodes.Add(MeshGeometry.MeshUid);
+						}
+					}
+
+					if (MeshUidsPerLodIndex.Num() > 0)
+					{
+						UInterchangeSkeletalMeshFactoryNode* SkeletalMeshFactoryNode = CreateSkeletalMeshFactoryNode(SkeletonRootUid, MeshUidsPerLodIndex);
+						SetSkeletalMeshDependencies(SkeletonRootUid, SkeletalMeshFactoryNode);
+						SkeletonFactoryNodes.Add(SkeletonFactoryNode);
+						SkeletalMeshFactoryNodes.Add(SkeletalMeshFactoryNode);
+						bFoundInstances = true;
+					}
 				}
+				return bFoundInstances;
+			};
+
+			bool bFoundMeshes = false;
+			if (bBakeMeshes)
+			{
+				PipelineMeshesUtilities->GetCombinedSkinnedMeshInstances(MeshUidsPerSkeletonRootUid);
+				const bool bUseMeshInstance = true;
+				bFoundMeshes = CreatePerSkeletonRootUidCombinedSkinnedMesh(bUseMeshInstance);
+			}
+
+			if (!bFoundMeshes)
+			{
+				MeshUidsPerSkeletonRootUid.Empty();
+				PipelineMeshesUtilities->GetCombinedSkinnedMeshGeometries(MeshUidsPerSkeletonRootUid);
+				const bool bUseMeshInstance = false;
+				CreatePerSkeletonRootUidCombinedSkinnedMesh(bUseMeshInstance);
 			}
 		}
-		
-		TMap<FString, TArray<FString>> MesheUidsPerJointRootUid;
-		for (const UInterchangeMeshNode* SkinnedMeshNode : SkinnedMeshNodes)
+		else
 		{
-			if (SkinnedMeshNode->GetSkeletonDependeciesCount() == 0)
+			//////////////////////////////////////////////////////////////////////////
+			//Do not combined meshes
+			TArray<FString> MeshUids;
+			auto CreatePerSkeletonRootUidSkinnedMesh = [this, &MeshUids, &SetSkeletalMeshDependencies](const bool bUseInstanceMesh)
 			{
-				continue;
-			}
-			FString JointNodeUid;
-			SkinnedMeshNode->GetSkeletonDependency(0, JointNodeUid);
-			while (!SkeletonRootNodeUids.Contains(JointNodeUid))
+				bool bFoundInstances = false;
+				for (const FString& MeshUid : MeshUids)
+				{
+					//Every iteration is a skeletalmesh asset that combine all MeshInstances sharing the same skeleton root node
+					//The MeshUids can represent a SceneNode pointing on a MeshNode or directly a MeshNode;
+					TMap<int32, TArray<FString>> MeshUidsPerLodIndex;
+					FString SkeletonRootUid;
+					if (!(bUseInstanceMesh ? PipelineMeshesUtilities->IsValidMeshInstanceUid(MeshUid) : PipelineMeshesUtilities->IsValidMeshGeometryUid(MeshUid)))
+					{
+						continue;
+					}
+					SkeletonRootUid = (bUseInstanceMesh ? PipelineMeshesUtilities->GetMeshInstanceSkeletonRootUid(MeshUid) : PipelineMeshesUtilities->GetMeshGeometrySkeletonRootUid(MeshUid));
+					if (SkeletonRootUid.IsEmpty())
+					{
+						//Log an error
+						continue;
+					}
+					UInterchangeSkeletonFactoryNode* SkeletonFactoryNode = CreateSkeletonFactoryNode(SkeletonRootUid);
+					if (bUseInstanceMesh)
+					{
+						const FInterchangeMeshInstance& MeshInstance = PipelineMeshesUtilities->GetMeshInstanceByUid(MeshUid);
+						for (const TPair<int32, FInterchangeLodSceneNodeContainer>& LodIndexAndSceneNodeContainer : MeshInstance.SceneNodePerLodIndex)
+						{
+							const int32 LodIndex = LodIndexAndSceneNodeContainer.Key;
+							const FInterchangeLodSceneNodeContainer& SceneNodeContainer = LodIndexAndSceneNodeContainer.Value;
+							TArray<FString>& TranslatedNodes = MeshUidsPerLodIndex.FindOrAdd(LodIndex);
+							for (const UInterchangeSceneNode* SceneNode : SceneNodeContainer.SceneNodes)
+							{
+								TranslatedNodes.Add(SceneNode->GetUniqueID());
+							}
+						}
+					}
+					else
+					{
+						const FInterchangeMeshGeometry& MeshGeometry = PipelineMeshesUtilities->GetMeshGeometryByUid(MeshUid);
+						const int32 LodIndex = 0;
+						TArray<FString>& TranslatedNodes = MeshUidsPerLodIndex.FindOrAdd(LodIndex);
+						TranslatedNodes.Add(MeshGeometry.MeshUid);
+					}
+					if (MeshUidsPerLodIndex.Num() > 0)
+					{
+						UInterchangeSkeletalMeshFactoryNode* SkeletalMeshFactoryNode = CreateSkeletalMeshFactoryNode(SkeletonRootUid, MeshUidsPerLodIndex);
+						SetSkeletalMeshDependencies(SkeletonRootUid, SkeletalMeshFactoryNode);
+						SkeletonFactoryNodes.Add(SkeletonFactoryNode);
+						SkeletalMeshFactoryNodes.Add(SkeletalMeshFactoryNode);
+						bFoundInstances = true;
+					}
+				}
+				return bFoundInstances;
+			};
+			
+			bool bFoundMeshes = false;
+			if (bBakeMeshes)
 			{
-				JointNodeUid = BaseNodeContainer->GetNode(JointNodeUid)->GetParentUid();
+				PipelineMeshesUtilities->GetAllSkinnedMeshInstance(MeshUids);
+				const bool bUseMeshInstance = true;
+				bFoundMeshes = CreatePerSkeletonRootUidSkinnedMesh(bUseMeshInstance);
 			}
-			if (SkeletonRootNodeUids.Contains(JointNodeUid))
+
+			if (!bFoundMeshes)
 			{
-				TArray<FString>& CombinedMeshNodes = MesheUidsPerJointRootUid.FindOrAdd(JointNodeUid);
-				CombinedMeshNodes.Add(SkinnedMeshNode->GetUniqueID());
+				MeshUids.Empty();
+				PipelineMeshesUtilities->GetAllSkinnedMeshGeometry(MeshUids);
+				const bool bUseMeshInstance = false;
+				CreatePerSkeletonRootUidSkinnedMesh(bUseMeshInstance);
 			}
 		}
+	}
+
+	const UClass* SkeletalMeshFactoryNodeClass = UInterchangeSkeletalMeshFactoryNode::StaticClass();
+	TArray<FString> SkeletalMeshNodeUids;
+	BaseNodeContainer->GetNodes(SkeletalMeshFactoryNodeClass, SkeletalMeshNodeUids);
+
+// 	const UClass* StaticMeshFactoryNodeClass = UInterchangeStaticMeshFactoryNode::StaticClass();
+// 	TArray<FString> StaticMeshNodeUids;
+// 	BaseNodeContainer->GetNodes(StaticMeshFactoryNodeClass, StaticMeshNodeUids);
+
+	//TODO count also the imported animations
 
-		const bool bLocalUseSourceNameForAsset = bUseSourceNameForAsset && (SourceDatas.Num() == 1) && MesheUidsPerJointRootUid.Num() == 1 && SkinnedMeshNodes.Num() == MeshNodes.Num();
-		for (TPair<FString, TArray<FString>> JointRootUidAndMeshes : MesheUidsPerJointRootUid)
+	const int32 MeshesAndAnimsImportedNodeCount = SkeletalMeshNodeUids.Num(); // + StaticMeshNodeUids.Num();
+	if (bUseSourceNameForAsset && MeshesAndAnimsImportedNodeCount == 1)
+	{
+		if (SkeletalMeshNodeUids.Num() > 0)
 		{
-			SkeletonFactoryNodes.Add(CreateSkeletonFactoryNode(JointRootUidAndMeshes.Key));
-			SkeletalMeshFactoryNodes.Add(CreateSkeletalMeshFactoryNode(JointRootUidAndMeshes.Key, JointRootUidAndMeshes.Value, bLocalUseSourceNameForAsset));
+			UInterchangeSkeletalMeshFactoryNode* SkeletalMeshNode = Cast<UInterchangeSkeletalMeshFactoryNode>(BaseNodeContainer->GetNode(SkeletalMeshNodeUids[0]));
+			const FString DisplayLabelName = FPaths::GetBaseFilename(SourceDatas[0]->GetFilename());
+			SkeletalMeshNode->SetDisplayLabel(DisplayLabelName);
 		}
 	}
-	else
-	{
 
-	}
 	return true;
 }
 
-TObjectPtr<UInterchangeTextureFactoryNode> UInterchangeGenericAssetsPipeline::CreateTexture2DFactoryNode(const TObjectPtr<UInterchangeTextureNode> TextureNode)
+UInterchangeTextureFactoryNode* UInterchangeGenericAssetsPipeline::CreateTexture2DFactoryNode(const UInterchangeTextureNode* TextureNode)
 {
 	FString DisplayLabel = TextureNode->GetDisplayLabel();
 	FString NodeUid = UInterchangeTextureFactoryNode::GetTextureFactoryNodeUidFromTextureNodeUid(TextureNode->GetUniqueID());
@@ -159,7 +294,7 @@ TObjectPtr<UInterchangeTextureFactoryNode> UInterchangeGenericAssetsPipeline::Cr
 	}
 	else
 	{
-		TextureFactoryNode = NewObject<UInterchangeTextureFactoryNode>(BaseNodeContainer.Get(), NAME_None);
+		TextureFactoryNode = NewObject<UInterchangeTextureFactoryNode>(BaseNodeContainer, NAME_None);
 		if (!ensure(TextureFactoryNode))
 		{
 			return nullptr;
@@ -173,7 +308,7 @@ TObjectPtr<UInterchangeTextureFactoryNode> UInterchangeGenericAssetsPipeline::Cr
 	return TextureFactoryNode;
 }
 
-TObjectPtr<UInterchangeMaterialFactoryNode> UInterchangeGenericAssetsPipeline::CreateMaterialFactoryNode(const TObjectPtr<UInterchangeMaterialNode> MaterialNode)
+UInterchangeMaterialFactoryNode* UInterchangeGenericAssetsPipeline::CreateMaterialFactoryNode(const UInterchangeMaterialNode* MaterialNode)
 {
 	FString DisplayLabel = MaterialNode->GetDisplayLabel();
 	FString NodeUid = UInterchangeMaterialFactoryNode::GetMaterialFactoryNodeUidFromMaterialNodeUid(MaterialNode->GetUniqueID());
@@ -189,7 +324,7 @@ TObjectPtr<UInterchangeMaterialFactoryNode> UInterchangeGenericAssetsPipeline::C
 	}
 	else
 	{
-		MaterialFactoryNode = NewObject<UInterchangeMaterialFactoryNode>(BaseNodeContainer.Get(), NAME_None);
+		MaterialFactoryNode = NewObject<UInterchangeMaterialFactoryNode>(BaseNodeContainer, NAME_None);
 		if (!ensure(MaterialFactoryNode))
 		{
 			return nullptr;
@@ -216,7 +351,7 @@ TObjectPtr<UInterchangeMaterialFactoryNode> UInterchangeGenericAssetsPipeline::C
 	return MaterialFactoryNode;
 }
 
-TObjectPtr<UInterchangeSkeletonFactoryNode> UInterchangeGenericAssetsPipeline::CreateSkeletonFactoryNode(const FString& RootJointUid)
+UInterchangeSkeletonFactoryNode* UInterchangeGenericAssetsPipeline::CreateSkeletonFactoryNode(const FString& RootJointUid)
 {
 	const UInterchangeBaseNode* RootJointNode = BaseNodeContainer->GetNode(RootJointUid);
 	if (!RootJointNode)
@@ -239,7 +374,7 @@ TObjectPtr<UInterchangeSkeletonFactoryNode> UInterchangeGenericAssetsPipeline::C
 	}
 	else
 	{
-		SkeletonFactoryNode = NewObject<UInterchangeSkeletonFactoryNode>(BaseNodeContainer.Get(), NAME_None);
+		SkeletonFactoryNode = NewObject<UInterchangeSkeletonFactoryNode>(BaseNodeContainer, NAME_None);
 		if (!ensure(SkeletonFactoryNode))
 		{
 			return nullptr;
@@ -251,7 +386,7 @@ TObjectPtr<UInterchangeSkeletonFactoryNode> UInterchangeGenericAssetsPipeline::C
 	return SkeletonFactoryNode;
 }
 
-TObjectPtr<UInterchangeSkeletalMeshFactoryNode> UInterchangeGenericAssetsPipeline::CreateSkeletalMeshFactoryNode(const FString& RootJointUid, TArray<FString>& MeshNodeUids, bool bLocalUseSourceNameForAsset)
+UInterchangeSkeletalMeshFactoryNode* UInterchangeGenericAssetsPipeline::CreateSkeletalMeshFactoryNode(const FString& RootJointUid, const TMap<int32, TArray<FString>> MeshUidsPerLodIndex)
 {
 	//Get the skeleton factory node
 	const UInterchangeBaseNode* RootJointNode = BaseNodeContainer->GetNode(RootJointUid);
@@ -267,118 +402,97 @@ TObjectPtr<UInterchangeSkeletalMeshFactoryNode> UInterchangeGenericAssetsPipelin
 		return nullptr;
 	}
 	
-	if (MeshNodeUids.Num() == 0)
+	if (MeshUidsPerLodIndex.Num() == 0)
 	{
 		return nullptr;
 	}
-	const UInterchangeMeshNode* FirstMeshNode = Cast<UInterchangeMeshNode>(BaseNodeContainer->GetNode(MeshNodeUids[0]));
-	if (!FirstMeshNode)
-	{
-		return nullptr;
-	}
-
-	//Create the skeletal mesh factory node, name it according to the first mesh node compositing the meshes
-	FString DisplayLabel = FirstMeshNode->GetDisplayLabel();
-	if (bLocalUseSourceNameForAsset && SourceDatas.Num() > 0)
-	{
-		DisplayLabel = FPaths::GetBaseFilename(SourceDatas[0]->GetFilename());
-	}
-	const FString SkeletalMeshUid = TEXT("\\SkeletalMesh\\") + RootJointNode->GetUniqueID();
-
-	UInterchangeSkeletalMeshFactoryNode* SkeletalMeshFactoryNode = NewObject<UInterchangeSkeletalMeshFactoryNode>(BaseNodeContainer.Get(), NAME_None);
-	if (!ensure(SkeletalMeshFactoryNode))
-	{
-		return nullptr;
-	}
-	SkeletalMeshFactoryNode->InitializeSkeletalMeshNode(SkeletalMeshUid, DisplayLabel, USkeletalMesh::StaticClass()->GetName());
-	SkeletalMeshFactoryNode->SetFactoryDependencyUid(SkeletonUid);
-	BaseNodeContainer->AddNode(SkeletalMeshFactoryNode);
-
-	TMap< FString, TMap<int32, TArray<FString>>> SceneInstanceUidsPerLodIndexPerMeshNodeUid;
-	for (const FString& MeshNodeUid : MeshNodeUids)
+	
+	auto GetFirstNodeInfo = [this, &MeshUidsPerLodIndex](const int32 Index, FString& OutFirstMeshNodeUid, int32& OutSceneNodeCount)->const UInterchangeBaseNode*
 	{
-		const UInterchangeMeshNode* MeshNode = Cast<UInterchangeMeshNode>(BaseNodeContainer->GetNode(MeshNodeUid));
-		if (!ensure(MeshNode))
+		OutSceneNodeCount = 0;
+		if (!ensure(Index >= 0 && MeshUidsPerLodIndex.Num() > Index))
 		{
-			continue;
+			//Log an error
+			return nullptr;
 		}
-		//Get all scene nodes that instance this mesh node with the LOD index in case its under a LodGroup
-		TMap<int32, TArray<FString>>& SceneInstanceUidsPerLodIndex = SceneInstanceUidsPerLodIndexPerMeshNodeUid.FindOrAdd(MeshNodeUid);
+		for (const TPair<int32, TArray<FString>>& LodIndexAndMeshUids : MeshUidsPerLodIndex)
 		{
-			TArray<FString> MeshSceneNodeInstanceUids;
-			MeshNode->GetSceneInstanceUids(MeshSceneNodeInstanceUids);
-			for (const FString& SceneNodeUid : MeshSceneNodeInstanceUids)
+			if (Index == LodIndexAndMeshUids.Key)
 			{
-				if (UInterchangeSceneNode* MeshSceneNode = Cast<UInterchangeSceneNode>(BaseNodeContainer->GetNode(SceneNodeUid)))
+				const TArray<FString>& MeshUids = LodIndexAndMeshUids.Value;
+				if (MeshUids.Num() > 0)
 				{
-					UInterchangeSceneNode* ParentMeshSceneNode = Cast<UInterchangeSceneNode>(BaseNodeContainer->GetNode(MeshSceneNode->GetParentUid()));
-					if (!ParentMeshSceneNode)
+					const FString& MeshUid = MeshUids[0];
+					UInterchangeMeshNode* MeshNode = Cast<UInterchangeMeshNode>(BaseNodeContainer->GetNode(MeshUid));
+					if (MeshNode)
 					{
-						continue;
+						OutFirstMeshNodeUid = MeshUid;
+						return MeshNode;
 					}
-					int32 SceneNodeLodIndex = 0;
-					FString LastChildUid = MeshSceneNode->GetUniqueID();
-					do
+					UInterchangeSceneNode* SceneNode = Cast<UInterchangeSceneNode>(BaseNodeContainer->GetNode(MeshUid));
+					if (SceneNode)
 					{
-						if (MeshSceneNode->IsSpecializedTypeContains(UE::Interchange::FSceneNodeStaticData::GetLodGroupSpecializeTypeString()))
+						FString MeshNodeUid;
+						if (SceneNode->GetCustomMeshDependencyUid(MeshNodeUid))
 						{
-							TArray<FString> LodGroupChildrens = BaseNodeContainer->GetNodeChildrenUids(ParentMeshSceneNode->GetUniqueID());
-							for (int32 LodIndex = 0; LodIndex < LodGroupChildrens.Num(); ++LodIndex)
-							{
-								const FString& ChildrenUid = LodGroupChildrens[LodIndex];
-								if (ChildrenUid.Equals(LastChildUid))
-								{
-									SceneNodeLodIndex = LodIndex;
-									break;
-								}
-							}
-							break;
+							OutSceneNodeCount = MeshUids.Num();
+							OutFirstMeshNodeUid = MeshNodeUid;
+							return SceneNode;
 						}
-						LastChildUid = ParentMeshSceneNode->GetUniqueID();
-						ParentMeshSceneNode = Cast<UInterchangeSceneNode>(BaseNodeContainer->GetNode(ParentMeshSceneNode->GetParentUid()));
-					} while (ParentMeshSceneNode);
-
-					//Apply bImportLod option, it will not add LOD
-					if (bImportLods || SceneNodeLodIndex == 0)
-					{
-						TArray<FString>& SceneInstanceUids = SceneInstanceUidsPerLodIndex.FindOrAdd(SceneNodeLodIndex);
-						SceneInstanceUids.AddUnique(SceneNodeUid);
 					}
 				}
+				//We found the lod but there is no valid Mesh node to return the Uid
+				break;
 			}
 		}
-	}
+		return nullptr;
+	};
 
-	for (const TPair< FString, TMap<int32, TArray<FString>>>& MeshNodeUidAndSceneInstanceUidsPerLodIndex : SceneInstanceUidsPerLodIndexPerMeshNodeUid)
+	FString FirstMeshNodeUid;
+	const int32 BaseLodIndex = 0;
+	int32 SceneNodeCount = 0;
+	const UInterchangeBaseNode* InterchangeBaseNode = GetFirstNodeInfo(BaseLodIndex, FirstMeshNodeUid, SceneNodeCount);
+	if (!InterchangeBaseNode)
 	{
-		const FString& MeshNodeUid = MeshNodeUidAndSceneInstanceUidsPerLodIndex.Key;
-		const UInterchangeMeshNode* MeshNode = Cast<UInterchangeMeshNode>(BaseNodeContainer->GetNode(MeshNodeUid));
-		if (!ensure(MeshNode))
-		{
-			continue;
-		}
+		//Log an error
+		return nullptr;
+	}
+	const UInterchangeSceneNode* FirstSceneNode = Cast<UInterchangeSceneNode>(InterchangeBaseNode);
+	const UInterchangeMeshNode* FirstMeshNode = Cast<UInterchangeMeshNode>(BaseNodeContainer->GetNode(FirstMeshNodeUid));
 
-		const TMap<int32, TArray<FString>>& SceneInstanceUidsPerLodIndex = MeshNodeUidAndSceneInstanceUidsPerLodIndex.Value;
-		if (bBakeMeshes && SceneInstanceUidsPerLodIndex.Num() > 0)
-		{
-			//Baked and instanced meshes code path
-			AddLodDataToSkeletalMesh(SkeletonFactoryNode, SkeletalMeshFactoryNode, SceneInstanceUidsPerLodIndex);
-		}
-		else
+	//Create the skeletal mesh factory node, name it according to the first mesh node compositing the meshes
+	FString DisplayLabel = FirstMeshNode->GetDisplayLabel();
+	FString SkeletalMeshUid_MeshNamePart = FirstMeshNodeUid;
+	if(FirstSceneNode)
+	{
+		//If we are instancing one scene node, we want to use it to name the mesh
+		if (SceneNodeCount == 1)
 		{
-			//Unbaked or un-instanced meshes code path
-			AddLodDataToSkeletalMesh(SkeletonFactoryNode, SkeletalMeshFactoryNode, MeshNode);
+			DisplayLabel = FirstSceneNode->GetDisplayLabel();
 		}
+		//Use the first scene node uid this skeletalmesh reference, add backslash since this uid is not asset typed (\\Mesh\\) like FirstMeshNodeUid
+		SkeletalMeshUid_MeshNamePart = TEXT("\\") + FirstSceneNode->GetUniqueID();
 	}
+	const FString SkeletalMeshUid = TEXT("\\SkeletalMesh") + SkeletalMeshUid_MeshNamePart + SkeletonUid;
+	UInterchangeSkeletalMeshFactoryNode* SkeletalMeshFactoryNode = NewObject<UInterchangeSkeletalMeshFactoryNode>(BaseNodeContainer, NAME_None);
+	if (!ensure(SkeletalMeshFactoryNode))
+	{
+		return nullptr;
+	}
+	SkeletalMeshFactoryNode->InitializeSkeletalMeshNode(SkeletalMeshUid, DisplayLabel, USkeletalMesh::StaticClass()->GetName());
+	SkeletalMeshFactoryNode->SetFactoryDependencyUid(SkeletonUid);
+	BaseNodeContainer->AddNode(SkeletalMeshFactoryNode);
+
+	AddLodDataToSkeletalMesh(SkeletonFactoryNode, SkeletalMeshFactoryNode, MeshUidsPerLodIndex);
 
 	return SkeletalMeshFactoryNode;
 }
 
-TObjectPtr<UInterchangeSkeletalMeshLodDataNode> UInterchangeGenericAssetsPipeline::CreateSkeletalMeshLodDataNode(const FString& NodeName, const FString& NodeUniqueID)
+UInterchangeSkeletalMeshLodDataNode* UInterchangeGenericAssetsPipeline::CreateSkeletalMeshLodDataNode(const FString& NodeName, const FString& NodeUniqueID)
 {
 	FString DisplayLabel(NodeName);
 	FString NodeUID(NodeUniqueID);
-	UInterchangeSkeletalMeshLodDataNode* SkeletalMeshLodDataNode = NewObject<UInterchangeSkeletalMeshLodDataNode>(BaseNodeContainer.Get(), NAME_None);
+	UInterchangeSkeletalMeshLodDataNode* SkeletalMeshLodDataNode = NewObject<UInterchangeSkeletalMeshLodDataNode>(BaseNodeContainer, NAME_None);
 	if (!ensure(SkeletalMeshLodDataNode))
 	{
 		//TODO log error
@@ -390,21 +504,21 @@ TObjectPtr<UInterchangeSkeletalMeshLodDataNode> UInterchangeGenericAssetsPipelin
 	return SkeletalMeshLodDataNode;
 }
 
-void UInterchangeGenericAssetsPipeline::AddLodDataToSkeletalMesh(const UInterchangeSkeletonFactoryNode* SkeletonFactoryNode, UInterchangeSkeletalMeshFactoryNode* SkeletalMeshFactoryNode, const TMap<int32, TArray<FString>>& SceneInstanceUidsPerLodIndex)
+void UInterchangeGenericAssetsPipeline::AddLodDataToSkeletalMesh(const UInterchangeSkeletonFactoryNode* SkeletonFactoryNode, UInterchangeSkeletalMeshFactoryNode* SkeletalMeshFactoryNode, const TMap<int32, TArray<FString>>& NodeUidsPerLodIndex)
 {
 	const FString SkeletalMeshUid = SkeletalMeshFactoryNode->GetUniqueID();
 	const FString SkeletonUid = SkeletonFactoryNode->GetUniqueID();
 	FString RootjointNodeUid;
 	SkeletonFactoryNode->GetCustomRootJointUid(RootjointNodeUid);
-	for (const TPair<int32, TArray<FString>>& LodIndexAndSceneInstanceUids : SceneInstanceUidsPerLodIndex)
+	for (const TPair<int32, TArray<FString>>& LodIndexAndNodeUids : NodeUidsPerLodIndex)
 	{
-		const int32 LodIndex = LodIndexAndSceneInstanceUids.Key;
-		const TArray<FString>& SceneNodeUids = LodIndexAndSceneInstanceUids.Value;
+		const int32 LodIndex = LodIndexAndNodeUids.Key;
+		const TArray<FString>& NodeUids = LodIndexAndNodeUids.Value;
 
 		//Create a lod data node with all the meshes for this LOD
 		const FString SkeletalMeshLodDataName = TEXT("LodData") + FString::FromInt(LodIndex);
-		const FString LODDataPrefix = TEXT("\\LodData") + (LodIndex > 0 ? FString::FromInt(LodIndex) : TEXT("")) + TEXT("\\");
-		const FString SkeletalMeshLodDataUniqueID = LODDataPrefix + RootjointNodeUid;
+		const FString LODDataPrefix = TEXT("\\LodData") + (LodIndex > 0 ? FString::FromInt(LodIndex) : TEXT(""));
+		const FString SkeletalMeshLodDataUniqueID = LODDataPrefix + SkeletalMeshUid + SkeletonUid;
 		//The LodData already exist
 		UInterchangeSkeletalMeshLodDataNode* LodDataNode = Cast<UInterchangeSkeletalMeshLodDataNode>(BaseNodeContainer->GetNode(SkeletalMeshLodDataUniqueID));
 		if (!LodDataNode)
@@ -415,51 +529,30 @@ void UInterchangeGenericAssetsPipeline::AddLodDataToSkeletalMesh(const UIntercha
 			LodDataNode->SetCustomSkeletonUid(SkeletonUid);
 			SkeletalMeshFactoryNode->AddLodDataUniqueId(SkeletalMeshLodDataUniqueID);
 		}
-		for (const FString& SceneNodeUid : SceneNodeUids)
+		for (const FString& NodeUid : NodeUids)
 		{
-			const UInterchangeSceneNode* SceneNode = Cast<UInterchangeSceneNode>(BaseNodeContainer->GetNode(SceneNodeUid));
 			TArray<FString> MaterialDependencies;
-			SceneNode->GetMaterialDependencyUids(MaterialDependencies);
+			if (const UInterchangeSceneNode* SceneNode = Cast<UInterchangeSceneNode>(BaseNodeContainer->GetNode(NodeUid)))
+			{
+				SceneNode->GetMaterialDependencyUids(MaterialDependencies);
+			}
+			else if (const UInterchangeMeshNode* MeshNode = Cast<UInterchangeMeshNode>(BaseNodeContainer->GetNode(NodeUid)))
+			{
+				MeshNode->GetMaterialDependencies(MaterialDependencies);
+			}
 			const int32 MaterialCount = MaterialDependencies.Num();
 			for (int32 MaterialIndex = 0; MaterialIndex < MaterialCount; ++MaterialIndex)
 			{
 				const FString MaterialFactoryNodeUid = UInterchangeMaterialFactoryNode::GetMaterialFactoryNodeUidFromMaterialNodeUid(MaterialDependencies[MaterialIndex]);
 				//Create a factory dependency so Material asset are import before the skeletal mesh asset
-				SkeletalMeshFactoryNode->SetFactoryDependencyUid(MaterialFactoryNodeUid);
+				TArray<FString> FactoryDependencies;
+				SkeletalMeshFactoryNode->GetFactoryDependencies(FactoryDependencies);
+				if (!FactoryDependencies.Contains(MaterialFactoryNodeUid))
+				{
+					SkeletalMeshFactoryNode->SetFactoryDependencyUid(MaterialFactoryNodeUid);
+				}
 			}
-			//Add the scene node has a reference to bake all mesh instance
-			LodDataNode->AddMeshUid(SceneNodeUid);
+			LodDataNode->AddMeshUid(NodeUid);
 		}
 	}
 }
-
-void UInterchangeGenericAssetsPipeline::AddLodDataToSkeletalMesh(const UInterchangeSkeletonFactoryNode* SkeletonFactoryNode, UInterchangeSkeletalMeshFactoryNode* SkeletalMeshFactoryNode, const UInterchangeMeshNode* MeshNode)
-{
-	const FString SkeletalMeshUid = SkeletalMeshFactoryNode->GetUniqueID();
-	const FString SkeletonUid = SkeletonFactoryNode->GetUniqueID();
-	FString RootjointNodeUid;
-	SkeletonFactoryNode->GetCustomRootJointUid(RootjointNodeUid);
-	//simply export the mesh node to lod 0
-	//Create a lod data node with all the meshes for this LOD
-	const FString SkeletalMeshLodDataName = TEXT("LodData0");
-	const FString SkeletalMeshLodDataUniqueID = TEXT("\\LodData\\") + RootjointNodeUid;
-	UInterchangeSkeletalMeshLodDataNode* LodDataNode = Cast<UInterchangeSkeletalMeshLodDataNode>(BaseNodeContainer->GetNode(SkeletalMeshLodDataUniqueID));
-	if (!LodDataNode)
-	{
-		//Add the data for the LOD (skeleton Unique ID and all the mesh node fbx path, so we can find them when we will create the payload data)
-		LodDataNode = CreateSkeletalMeshLodDataNode(SkeletalMeshLodDataName, SkeletalMeshLodDataUniqueID);
-		LodDataNode->SetParentUid(SkeletalMeshUid);
-		LodDataNode->SetCustomSkeletonUid(SkeletonUid);
-		SkeletalMeshFactoryNode->AddLodDataUniqueId(SkeletalMeshLodDataUniqueID);
-	}
-	TArray<FString> MaterialDependencies;
-	MeshNode->GetMaterialDependencies(MaterialDependencies);
-	const int32 MaterialCount = MaterialDependencies.Num();
-	for (int32 MaterialIndex = 0; MaterialIndex < MaterialCount; ++MaterialIndex)
-	{
-		const FString MaterialFactoryNodeUid = UInterchangeMaterialFactoryNode::GetMaterialFactoryNodeUidFromMaterialNodeUid(MaterialDependencies[MaterialIndex]);
-		//Create a factory dependency so Material asset are import before the skeletal mesh asset
-		SkeletalMeshFactoryNode->SetFactoryDependencyUid(MaterialFactoryNodeUid);
-	}
-	LodDataNode->AddMeshUid(MeshNode->GetUniqueID());
-}
\ No newline at end of file
diff --git a/Engine/Plugins/Runtime/Interchange/Source/Pipelines/Private/InterchangePipelineMeshesUtilities.cpp b/Engine/Plugins/Runtime/Interchange/Source/Pipelines/Private/InterchangePipelineMeshesUtilities.cpp
new file mode 100644
--- /dev/null
+++ b/Engine/Plugins/Runtime/Interchange/Source/Pipelines/Private/InterchangePipelineMeshesUtilities.cpp
@@ -0,0 +1,468 @@
+// Copyright Epic Games, Inc. All Rights Reserved.
+
+#pragma once
+
+#include "InterchangePipelineMeshesUtilities.h"
+
+#include "CoreMinimal.h"
+#include "Nodes/InterchangeBaseNode.h"
+#include "InterchangeMeshNode.h"
+#include "InterchangeSceneNode.h"
+#include "Nodes/InterchangeBaseNodeContainer.h"
+#include "UObject/Object.h"
+#include "UObject/ObjectMacros.h"
+
+
+UInterchangePipelineMeshesUtilities* UInterchangePipelineMeshesUtilities::CreateInterchangePipelineMeshesUtilities(UInterchangeBaseNodeContainer* BaseNodeContainer)
+{
+	check(BaseNodeContainer);
+	UInterchangePipelineMeshesUtilities* PipelineMeshesUtilities = NewObject<UInterchangePipelineMeshesUtilities>(GetTransientPackage(), NAME_None);
+	
+	TArray<FString> SkeletonRootNodeUids;
+	
+	//Find all translated node we need for this pipeline
+	BaseNodeContainer->IterateNodes([&PipelineMeshesUtilities, &BaseNodeContainer, &SkeletonRootNodeUids](const FString& NodeUid, UInterchangeBaseNode* Node)
+	{
+		switch (Node->GetnodeContainerType())
+		{
+			case EInterchangeNodeContainerType::NodeContainerType_TranslatedAsset:
+			{
+				if (UInterchangeMeshNode* MeshNode = Cast<UInterchangeMeshNode>(Node))
+				{
+					FInterchangeMeshGeometry& MeshGeometry = PipelineMeshesUtilities->MeshGeometriesPerMeshUid.FindOrAdd(NodeUid);
+					MeshGeometry.MeshUid = NodeUid;
+					MeshGeometry.MeshNode = MeshNode;
+				}
+			}
+			break;
+			case EInterchangeNodeContainerType::NodeContainerType_TranslatedScene:
+			{
+				if (UInterchangeSceneNode* SceneNode = Cast<UInterchangeSceneNode>(Node))
+				{
+					if (SceneNode->IsSpecializedTypeContains(UE::Interchange::FSceneNodeStaticData::GetJointSpecializeTypeString()))
+					{
+						const UInterchangeSceneNode* ParentJointNode = Cast<UInterchangeSceneNode>(BaseNodeContainer->GetNode(SceneNode->GetParentUid()));
+						if (!ParentJointNode || !ParentJointNode->IsSpecializedTypeContains(UE::Interchange::FSceneNodeStaticData::GetJointSpecializeTypeString()))
+						{
+							SkeletonRootNodeUids.Add(SceneNode->GetUniqueID());
+						}
+					}
+
+					FString MeshUid;
+					if (SceneNode->GetCustomMeshDependencyUid(MeshUid))
+					{
+						UInterchangeSceneNode* ParentMeshSceneNode = Cast<UInterchangeSceneNode>(BaseNodeContainer->GetNode(SceneNode->GetParentUid()));
+						if (ParentMeshSceneNode)
+						{
+							UInterchangeSceneNode* LodGroupNode = nullptr;
+							int32 LodIndex = 0;
+							FString LastChildUid = SceneNode->GetUniqueID();
+							do
+							{
+								if (ParentMeshSceneNode->IsSpecializedTypeContains(UE::Interchange::FSceneNodeStaticData::GetLodGroupSpecializeTypeString()))
+								{
+									LodGroupNode = ParentMeshSceneNode;
+									TArray<FString> LodGroupChildrens = BaseNodeContainer->GetNodeChildrenUids(ParentMeshSceneNode->GetUniqueID());
+									for (int32 ChildLodIndex = 0; ChildLodIndex < LodGroupChildrens.Num(); ++ChildLodIndex)
+									{
+										const FString& ChildrenUid = LodGroupChildrens[ChildLodIndex];
+										if (ChildrenUid.Equals(LastChildUid))
+										{
+											LodIndex = ChildLodIndex;
+											break;
+										}
+									}
+									break;
+								}
+								LastChildUid = ParentMeshSceneNode->GetUniqueID();
+								ParentMeshSceneNode = Cast<UInterchangeSceneNode>(BaseNodeContainer->GetNode(ParentMeshSceneNode->GetParentUid()));
+							} while (ParentMeshSceneNode);
+
+							FInterchangeMeshGeometry& MeshGeometry = PipelineMeshesUtilities->MeshGeometriesPerMeshUid.FindChecked(MeshUid);
+							if (LodGroupNode)
+							{
+								//We have a LOD
+								FInterchangeMeshInstance& MeshInstance = PipelineMeshesUtilities->MeshInstancesPerMeshInstanceUid.FindOrAdd(LodGroupNode->GetUniqueID());
+								if (MeshInstance.LodGroupNode != nullptr)
+								{
+									//This LodGroup was already created, verify everything is ok
+									checkSlow(MeshInstance.LodGroupNode == LodGroupNode);
+									checkSlow(MeshInstance.MeshInstanceUid.Equals(LodGroupNode->GetUniqueID()));
+								}
+								else
+								{
+									MeshInstance.LodGroupNode = LodGroupNode;
+									MeshInstance.MeshInstanceUid = LodGroupNode->GetUniqueID();
+								}
+								FInterchangeLodSceneNodeContainer& InstancedSceneNodes = MeshInstance.SceneNodePerLodIndex.FindOrAdd(LodIndex);
+								InstancedSceneNodes.SceneNodes.AddUnique(SceneNode);
+								MeshGeometry.ReferencingMeshInstanceUids.Add(MeshInstance.MeshInstanceUid);
+								MeshInstance.ReferencingMeshGeometryUids.Add(MeshUid);
+								MeshInstance.bReferenceSkinnedMesh |= MeshGeometry.MeshNode->IsSkinnedMesh();
+							}
+							else
+							{
+								FInterchangeMeshInstance& MeshInstance = PipelineMeshesUtilities->MeshInstancesPerMeshInstanceUid.FindOrAdd(NodeUid);
+								MeshInstance.LodGroupNode = nullptr;
+								MeshInstance.MeshInstanceUid = NodeUid;
+								FInterchangeLodSceneNodeContainer& InstancedSceneNodes = MeshInstance.SceneNodePerLodIndex.FindOrAdd(LodIndex);
+								InstancedSceneNodes.SceneNodes.AddUnique(SceneNode);
+								MeshGeometry.ReferencingMeshInstanceUids.Add(MeshInstance.MeshInstanceUid);
+								MeshInstance.ReferencingMeshGeometryUids.Add(MeshUid);
+								MeshInstance.bReferenceSkinnedMesh |= MeshGeometry.MeshNode->IsSkinnedMesh();
+							}
+						}
+					}
+				}
+			}
+			break;
+		}
+	});
+
+	//Fill the SkeletonRootUidPerMeshUid data
+	for (const TPair<FString, FInterchangeMeshGeometry>& MeshGeometryUidAndMeshGeometry : PipelineMeshesUtilities->MeshGeometriesPerMeshUid)
+	{
+		const FInterchangeMeshGeometry& MeshGeometry = MeshGeometryUidAndMeshGeometry.Value;
+		if (!MeshGeometry.MeshNode->IsSkinnedMesh() || PipelineMeshesUtilities->SkeletonRootUidPerMeshUid.Contains(MeshGeometry.MeshUid))
+		{
+			continue;
+		}
+		const UInterchangeMeshNode* SkinnedMeshNode = MeshGeometry.MeshNode;
+		if (!SkinnedMeshNode || SkinnedMeshNode->GetSkeletonDependeciesCount() == 0)
+		{
+			continue;
+		}
+		//Find the root joint for this MeshGeometry
+		FString JointNodeUid;
+		SkinnedMeshNode->GetSkeletonDependency(0, JointNodeUid);
+		while (!JointNodeUid.Equals(UInterchangeBaseNode::InvalidNodeUid()) && !SkeletonRootNodeUids.Contains(JointNodeUid))
+		{
+			JointNodeUid = BaseNodeContainer->GetNode(JointNodeUid)->GetParentUid();
+		}
+		//Add the MeshGeometry to the map per joint uid
+		if (SkeletonRootNodeUids.Contains(JointNodeUid))
+		{
+			PipelineMeshesUtilities->SkeletonRootUidPerMeshUid.Add(MeshGeometry.MeshUid, JointNodeUid);
+		}
+	}
+
+	return PipelineMeshesUtilities;
+}
+
+void UInterchangePipelineMeshesUtilities::GetAllMeshInstanceUids(TArray<FString>& MeshInstanceUids) const
+{
+	MeshInstancesPerMeshInstanceUid.GetKeys(MeshInstanceUids);
+}
+
+void UInterchangePipelineMeshesUtilities::IterateAllMeshInstance(TFunctionRef<void(const FInterchangeMeshInstance&)> IterationLambda) const
+{
+	for (const TPair<FString, FInterchangeMeshInstance>& MeshInstanceUidAndMeshInstance : MeshInstancesPerMeshInstanceUid)
+	{
+		const FInterchangeMeshInstance& MeshInstance = MeshInstanceUidAndMeshInstance.Value;
+		IterationLambda(MeshInstance);
+	}
+}
+
+void UInterchangePipelineMeshesUtilities::GetAllSkinnedMeshInstance(TArray<FString>& MeshInstanceUids) const
+{
+	MeshInstanceUids.Empty(MeshInstancesPerMeshInstanceUid.Num());
+	for (const TPair<FString, FInterchangeMeshInstance>& MeshInstanceUidAndMeshInstance : MeshInstancesPerMeshInstanceUid)
+	{
+		const FInterchangeMeshInstance& MeshInstance = MeshInstanceUidAndMeshInstance.Value;
+		if(MeshInstance.bReferenceSkinnedMesh)
+		{
+			MeshInstanceUids.Add(MeshInstance.MeshInstanceUid);
+		}
+	}
+}
+
+void UInterchangePipelineMeshesUtilities::IterateAllSkinnedMeshInstance(TFunctionRef<void(const FInterchangeMeshInstance&)> IterationLambda) const
+{
+	for (const TPair<FString, FInterchangeMeshInstance>& MeshInstanceUidAndMeshInstance : MeshInstancesPerMeshInstanceUid)
+	{
+		const FInterchangeMeshInstance& MeshInstance = MeshInstanceUidAndMeshInstance.Value;
+		if (MeshInstance.bReferenceSkinnedMesh)
+		{
+			IterationLambda(MeshInstance);
+		}
+	}
+}
+
+void UInterchangePipelineMeshesUtilities::GetAllStaticMeshInstance(TArray<FString>& MeshInstanceUids) const
+{
+	MeshInstanceUids.Empty(MeshInstancesPerMeshInstanceUid.Num());
+	for (const TPair<FString, FInterchangeMeshInstance>& MeshInstanceUidAndMeshInstance : MeshInstancesPerMeshInstanceUid)
+	{
+		const FInterchangeMeshInstance& MeshInstance = MeshInstanceUidAndMeshInstance.Value;
+		if (!MeshInstance.bReferenceSkinnedMesh)
+		{
+			MeshInstanceUids.Add(MeshInstance.MeshInstanceUid);
+		}
+	}
+}
+
+void UInterchangePipelineMeshesUtilities::IterateAllStaticMeshInstance(TFunctionRef<void(const FInterchangeMeshInstance&)> IterationLambda) const
+{
+	for (const TPair<FString, FInterchangeMeshInstance>& MeshInstanceUidAndMeshInstance : MeshInstancesPerMeshInstanceUid)
+	{
+		const FInterchangeMeshInstance& MeshInstance = MeshInstanceUidAndMeshInstance.Value;
+		if (!MeshInstance.bReferenceSkinnedMesh)
+		{
+			IterationLambda(MeshInstance);
+		}
+	}
+}
+
+void UInterchangePipelineMeshesUtilities::GetAllMeshGeometry(TArray<FString>& MeshGeometryUids) const
+{
+	MeshGeometriesPerMeshUid.GetKeys(MeshGeometryUids);
+}
+
+void UInterchangePipelineMeshesUtilities::IterateAllMeshGeometry(TFunctionRef<void(const FInterchangeMeshGeometry&)> IterationLambda) const
+{
+	for (const TPair<FString, FInterchangeMeshGeometry>& MeshGeometryUidAndMeshGeometry : MeshGeometriesPerMeshUid)
+	{
+		const FInterchangeMeshGeometry& MeshGeometry = MeshGeometryUidAndMeshGeometry.Value;
+		IterationLambda(MeshGeometry);
+	}
+}
+
+void UInterchangePipelineMeshesUtilities::GetAllSkinnedMeshGeometry(TArray<FString>& MeshGeometryUids) const
+{
+	MeshGeometryUids.Empty(MeshGeometriesPerMeshUid.Num());
+	for (const TPair<FString, FInterchangeMeshGeometry>& MeshGeometryUidAndMeshGeometry : MeshGeometriesPerMeshUid)
+	{
+		const FInterchangeMeshGeometry& MeshGeometry = MeshGeometryUidAndMeshGeometry.Value;
+		if (MeshGeometry.MeshNode->IsSkinnedMesh())
+		{
+			MeshGeometryUids.Add(MeshGeometry.MeshUid);
+		}
+	}
+}
+
+void UInterchangePipelineMeshesUtilities::IterateAllSkinnedMeshGeometry(TFunctionRef<void(const FInterchangeMeshGeometry&)> IterationLambda) const
+{
+	for (const TPair<FString, FInterchangeMeshGeometry>& MeshGeometryUidAndMeshGeometry : MeshGeometriesPerMeshUid)
+	{
+		const FInterchangeMeshGeometry& MeshGeometry = MeshGeometryUidAndMeshGeometry.Value;
+		if (MeshGeometry.MeshNode->IsSkinnedMesh())
+		{
+			IterationLambda(MeshGeometry);
+		}
+	}
+}
+
+void UInterchangePipelineMeshesUtilities::GetAllStaticMeshGeometry(TArray<FString>& MeshGeometryUids) const
+{
+	MeshGeometryUids.Empty(MeshGeometriesPerMeshUid.Num());
+	for (const TPair<FString, FInterchangeMeshGeometry>& MeshGeometryUidAndMeshGeometry : MeshGeometriesPerMeshUid)
+	{
+		const FInterchangeMeshGeometry& MeshGeometry = MeshGeometryUidAndMeshGeometry.Value;
+		if (!MeshGeometry.MeshNode->IsSkinnedMesh())
+		{
+			MeshGeometryUids.Add(MeshGeometry.MeshUid);
+		}
+	}
+}
+
+void UInterchangePipelineMeshesUtilities::IterateAllStaticMeshGeometry(TFunctionRef<void(const FInterchangeMeshGeometry&)> IterationLambda) const
+{
+	for (const TPair<FString, FInterchangeMeshGeometry>& MeshGeometryUidAndMeshGeometry : MeshGeometriesPerMeshUid)
+	{
+		const FInterchangeMeshGeometry& MeshGeometry = MeshGeometryUidAndMeshGeometry.Value;
+		if (!MeshGeometry.MeshNode->IsSkinnedMesh())
+		{
+			IterationLambda(MeshGeometry);
+		}
+	}
+}
+
+void UInterchangePipelineMeshesUtilities::GetAllMeshGeometryNotInstanced(TArray<FString>& MeshGeometryUids) const
+{
+	MeshGeometryUids.Empty(MeshGeometriesPerMeshUid.Num());
+	for (const TPair<FString, FInterchangeMeshGeometry>& MeshGeometryUidAndMeshGeometry : MeshGeometriesPerMeshUid)
+	{
+		const FInterchangeMeshGeometry& MeshGeometry = MeshGeometryUidAndMeshGeometry.Value;
+		if (MeshGeometry.ReferencingMeshInstanceUids.Num() == 0)
+		{
+			MeshGeometryUids.Add(MeshGeometry.MeshUid);
+		}
+	}
+}
+
+void UInterchangePipelineMeshesUtilities::IterateAllMeshGeometryNotIntanced(TFunctionRef<void(const FInterchangeMeshGeometry&)> IterationLambda) const
+{
+	for (const TPair<FString, FInterchangeMeshGeometry>& MeshGeometryUidAndMeshGeometry : MeshGeometriesPerMeshUid)
+	{
+		const FInterchangeMeshGeometry& MeshGeometry = MeshGeometryUidAndMeshGeometry.Value;
+		if (MeshGeometry.ReferencingMeshInstanceUids.Num() == 0)
+		{
+			IterationLambda(MeshGeometry);
+		}
+	}
+}
+
+bool UInterchangePipelineMeshesUtilities::IsValidMeshInstanceUid(const FString& MeshInstanceUid) const
+{
+	return MeshInstancesPerMeshInstanceUid.Contains(MeshInstanceUid);
+}
+
+const FInterchangeMeshInstance& UInterchangePipelineMeshesUtilities::GetMeshInstanceByUid(const FString& MeshInstanceUid) const
+{
+	return MeshInstancesPerMeshInstanceUid.FindChecked(MeshInstanceUid);
+}
+
+bool UInterchangePipelineMeshesUtilities::IsValidMeshGeometryUid(const FString& MeshGeometryUid) const
+{
+	return MeshGeometriesPerMeshUid.Contains(MeshGeometryUid);
+}
+
+const FInterchangeMeshGeometry& UInterchangePipelineMeshesUtilities::GetMeshGeometryByUid(const FString& MeshGeometryUid) const
+{
+	return MeshGeometriesPerMeshUid.FindChecked(MeshGeometryUid);
+}
+
+void UInterchangePipelineMeshesUtilities::GetAllMeshInstanceUidsUsingMeshGeometryUid(const FString& MeshGeometryUid, TArray<FString>& MeshInstanceUids) const
+{
+	const FInterchangeMeshGeometry& MeshGeometry = MeshGeometriesPerMeshUid.FindChecked(MeshGeometryUid);
+	MeshInstanceUids = MeshGeometry.ReferencingMeshInstanceUids;
+}
+
+void UInterchangePipelineMeshesUtilities::IterateAllMeshInstanceUsingMeshGeometry(const FString& MeshGeometryUid, TFunctionRef<void(const FInterchangeMeshInstance&)> IterationLambda) const
+{
+	const FInterchangeMeshGeometry& MeshGeometry = MeshGeometriesPerMeshUid.FindChecked(MeshGeometryUid);
+	for (const FString& MeshInstanceUid : MeshGeometry.ReferencingMeshInstanceUids)
+	{
+		const FInterchangeMeshInstance& MeshInstance = MeshInstancesPerMeshInstanceUid.FindChecked(MeshInstanceUid);
+		IterationLambda(MeshInstance);
+	}
+}
+
+void UInterchangePipelineMeshesUtilities::GetCombinedSkinnedMeshInstances(TMap<FString, TArray<FString>>& OutMeshInstanceUidsPerSkeletonRootUid) const
+{
+	for (const TPair<FString, FInterchangeMeshInstance>& MeshInstanceUidAndMeshInstance : MeshInstancesPerMeshInstanceUid)
+	{
+		const FInterchangeMeshInstance& MeshInstance = MeshInstanceUidAndMeshInstance.Value;
+		if (!MeshInstance.bReferenceSkinnedMesh)
+		{
+			continue;
+		}
+		//Find the root skeleton for this MeshInstance
+		FString SkeletonRootUid;
+		for (const FString& MeshGeometryUid : MeshInstance.ReferencingMeshGeometryUids)
+		{
+			if (const FString* SkeletonRootUidPtr = SkeletonRootUidPerMeshUid.Find(MeshGeometryUid))
+			{
+				if (SkeletonRootUid.IsEmpty())
+				{
+					SkeletonRootUid = *SkeletonRootUidPtr;
+				}
+				else if (!SkeletonRootUid.Equals(*SkeletonRootUidPtr))
+				{
+					//Log an error, this FInterchangeMeshInstance use more then one skeleton root node, we will not add this instance to the combined
+					SkeletonRootUid.Empty();
+					break;
+				}
+			}
+			else
+			{
+				//every skinned geometry should have a skeleton root node ???
+			}
+		}
+		
+		if (SkeletonRootUid.IsEmpty())
+		{
+			//Skip this MeshInstance
+			continue;
+		}
+		TArray<FString>& MeshInstanceUids = OutMeshInstanceUidsPerSkeletonRootUid.FindOrAdd(SkeletonRootUid);
+		MeshInstanceUids.Add(MeshInstanceUidAndMeshInstance.Key);
+	}
+}
+
+void UInterchangePipelineMeshesUtilities::GetCombinedSkinnedMeshGeometries(TMap<FString, TArray<FString>>& OutMeshGeometryUidsPerSkeletonRootUid) const
+{
+	for (const TPair<FString, FInterchangeMeshGeometry>& MeshGeometryUidAndMeshGeometry : MeshGeometriesPerMeshUid)
+	{
+		const FInterchangeMeshGeometry& MeshGeometry = MeshGeometryUidAndMeshGeometry.Value;
+		if (!MeshGeometry.MeshNode || !MeshGeometry.MeshNode->IsSkinnedMesh())
+		{
+			continue;
+		}
+		//Find the root skeleton for this MeshInstance
+		FString SkeletonRootUid;
+		if (const FString* SkeletonRootUidPtr = SkeletonRootUidPerMeshUid.Find(MeshGeometryUidAndMeshGeometry.Key))
+		{
+			if (SkeletonRootUid.IsEmpty())
+			{
+				SkeletonRootUid = *SkeletonRootUidPtr;
+			}
+		}
+		else
+		{
+			//every skinned geometry should have a skeleton root node ???
+		}
+
+		if (SkeletonRootUid.IsEmpty())
+		{
+			//Skip this MeshGeometry
+			continue;
+		}
+		TArray<FString>& MeshGeometryUids = OutMeshGeometryUidsPerSkeletonRootUid.FindOrAdd(SkeletonRootUid);
+		MeshGeometryUids.Add(MeshGeometryUidAndMeshGeometry.Key);
+	}
+}
+
+FString UInterchangePipelineMeshesUtilities::GetMeshInstanceSkeletonRootUid(const FString& MeshInstanceUid) const
+{
+	FString SkeletonRootUid;
+	if (IsValidMeshInstanceUid(MeshInstanceUid))
+	{
+		SkeletonRootUid = GetMeshInstanceSkeletonRootUid(GetMeshInstanceByUid(MeshInstanceUid));
+	}
+	return SkeletonRootUid;
+}
+
+FString UInterchangePipelineMeshesUtilities::GetMeshInstanceSkeletonRootUid(const FInterchangeMeshInstance& MeshInstance) const
+{
+	FString SkeletonRootUid;
+	if (MeshInstance.SceneNodePerLodIndex.Num() == 0)
+	{
+		return SkeletonRootUid;
+	}
+	const int32 BaseLodIndex = 0;
+	if (MeshInstance.SceneNodePerLodIndex[BaseLodIndex].SceneNodes.Num() > 0)
+	{
+		const UInterchangeSceneNode* SceneNode = MeshInstance.SceneNodePerLodIndex[BaseLodIndex].SceneNodes[0];
+		FString MeshNodeUid;
+		if (SceneNode->GetCustomMeshDependencyUid(MeshNodeUid))
+		{
+			if (SkeletonRootUidPerMeshUid.Contains(MeshNodeUid))
+			{
+				SkeletonRootUid = SkeletonRootUidPerMeshUid.FindChecked(MeshNodeUid);
+			}
+		}
+	}
+	return SkeletonRootUid;
+}
+
+FString UInterchangePipelineMeshesUtilities::GetMeshGeometrySkeletonRootUid(const FString& MeshGeometryUid) const
+{
+	FString SkeletonRootUid;
+	if (IsValidMeshGeometryUid(MeshGeometryUid))
+	{
+		const FInterchangeMeshGeometry& MeshGeometry = GetMeshGeometryByUid(MeshGeometryUid);
+		SkeletonRootUid = GetMeshGeometrySkeletonRootUid(MeshGeometry);
+	}
+	return SkeletonRootUid;
+}
+
+FString UInterchangePipelineMeshesUtilities::GetMeshGeometrySkeletonRootUid(const FInterchangeMeshGeometry& MeshGeometry) const
+{
+	FString SkeletonRootUid;
+	if (SkeletonRootUidPerMeshUid.Contains(MeshGeometry.MeshUid))
+	{
+		SkeletonRootUid = SkeletonRootUidPerMeshUid.FindChecked(MeshGeometry.MeshUid);
+	}
+	return SkeletonRootUid;
+}
\ No newline at end of file
diff --git a/Engine/Plugins/Runtime/Interchange/Source/Pipelines/Public/InterchangeGenericAssetsPipeline.h b/Engine/Plugins/Runtime/Interchange/Source/Pipelines/Public/InterchangeGenericAssetsPipeline.h
--- a/Engine/Plugins/Runtime/Interchange/Source/Pipelines/Public/InterchangeGenericAssetsPipeline.h
+++ b/Engine/Plugins/Runtime/Interchange/Source/Pipelines/Public/InterchangeGenericAssetsPipeline.h
@@ -16,6 +16,7 @@ class UInterchangeTextureFactoryNode;
 class UInterchangeMaterialNode;
 class UInterchangeMaterialFactoryNode;
 class UInterchangeMeshNode;
+class UInterchangePipelineMeshesUtilities;
 class UInterchangeSceneNode;
 class UInterchangeSkeletalMeshFactoryNode;
 class UInterchangeSkeletalMeshLodDataNode;
@@ -77,7 +78,7 @@ public:
 	UPROPERTY(EditAnywhere, BlueprintReadWrite, Category = COMMON_MESHES_CATEGORY, meta = (editcondition = "bBakeMeshes"))
 	bool bImportLods = true;
 
-	/** If enable, meshes will be baked with the hierarchy transform, if there is multiple instances, the mesh will incorporate all instance. */
+	/** If enable, meshes will be baked with the hierarchy transform. */
 	UPROPERTY(EditAnywhere, BlueprintReadWrite, AdvancedDisplay, Category = COMMON_MESHES_CATEGORY)
 	bool bBakeMeshes = true;
 
@@ -148,57 +149,63 @@ protected:
 	//virtual bool ExecuteExportPipeline(UInterchangeBaseNodeContainer* BaseNodeContainer) override;
 private:
 
-	TObjectPtr<UInterchangeBaseNodeContainer> BaseNodeContainer;
+	UInterchangeBaseNodeContainer* BaseNodeContainer;
 	TArray<const UInterchangeSourceData*> SourceDatas;
 
 	
 	/** Texture translated assets nodes */
-	TArray<TObjectPtr<UInterchangeTextureNode>> TextureNodes;
+	TArray<UInterchangeTextureNode*> TextureNodes;
 
 	/** Texture factory assets nodes */
-	TArray<TObjectPtr<UInterchangeTextureFactoryNode>> TextureFactoryNodes;
+	TArray<UInterchangeTextureFactoryNode*> TextureFactoryNodes;
 
-	TObjectPtr<UInterchangeTextureFactoryNode> CreateTexture2DFactoryNode(const TObjectPtr<UInterchangeTextureNode> TextureNode);
+	UInterchangeTextureFactoryNode* CreateTexture2DFactoryNode(const UInterchangeTextureNode* TextureNode);
 
 	
 	/** Material translated assets nodes */
-	TArray<TObjectPtr<UInterchangeMaterialNode>> MaterialNodes;
+	TArray<UInterchangeMaterialNode*> MaterialNodes;
 	
 	/** Material factory assets nodes */
-	TArray<TObjectPtr<UInterchangeMaterialFactoryNode>> MaterialFactoryNodes;
+	TArray<UInterchangeMaterialFactoryNode*> MaterialFactoryNodes;
 	
-	TObjectPtr<UInterchangeMaterialFactoryNode> CreateMaterialFactoryNode(const TObjectPtr<UInterchangeMaterialNode> MaterialNode);
+	UInterchangeMaterialFactoryNode* CreateMaterialFactoryNode(const UInterchangeMaterialNode* MaterialNode);
 
 	
 	/** Mesh translated assets nodes */
-	TArray<TObjectPtr<UInterchangeMeshNode>> MeshNodes;
+	TArray<UInterchangeMeshNode*> MeshNodes;
 	
 	/** Skeleton factory assets nodes */
-	TArray<TObjectPtr<UInterchangeSkeletonFactoryNode>> SkeletonFactoryNodes;
+	TArray<UInterchangeSkeletonFactoryNode*> SkeletonFactoryNodes;
 
 	/** Create a UInterchangeSkeletonFactorynode */
-	TObjectPtr<UInterchangeSkeletonFactoryNode> CreateSkeletonFactoryNode(const FString& RootJointUid);
+	UInterchangeSkeletonFactoryNode* CreateSkeletonFactoryNode(const FString& RootJointUid);
 	
 	/** Skeletal mesh factory assets nodes */
-	TArray<TObjectPtr<UInterchangeSkeletalMeshFactoryNode>> SkeletalMeshFactoryNodes;
+	TArray<UInterchangeSkeletalMeshFactoryNode*> SkeletalMeshFactoryNodes;
 	
 	/** Static mesh factory assets nodes */
-	//TArray<TObjectPtr<UInterchangeStaticMeshFactoryNode>> StaticMeshFactoryNodes;
+	//TArray<UInterchangeStaticMeshFactoryNode*> StaticMeshFactoryNodes;
 	
-	/** This function can create a UInterchangeSkeletalMeshFactoryNode */
-	TObjectPtr<UInterchangeSkeletalMeshFactoryNode> CreateSkeletalMeshFactoryNode(const FString& RootJointUid, TArray<FString>& MeshNodeUids, const bool bUseSourceNameForAsset);
+	/**
+	 * This function can create a UInterchangeSkeletalMeshFactoryNode
+	 * @param MeshUidsPerLodIndex - The MeshUids can represent a SceneNode pointing on a MeshNode or directly a MeshNode
+	 */
+	UInterchangeSkeletalMeshFactoryNode* CreateSkeletalMeshFactoryNode(const FString& RootJointUid, const TMap<int32, TArray<FString>> MeshUidsPerLodIndex);
 	
 	/** This function can create a UInterchangeSkeletalMeshLodDataNode which represent the LOD data need by the factory to create a lod mesh */
-	TObjectPtr<UInterchangeSkeletalMeshLodDataNode> CreateSkeletalMeshLodDataNode(const FString& NodeName, const FString& NodeUniqueID);
+	UInterchangeSkeletalMeshLodDataNode* CreateSkeletalMeshLodDataNode(const FString& NodeName, const FString& NodeUniqueID);
 
-	/** This function add all lod data node to the skeletal mesh, it handle instancing/combine and baking pipeline settings */
-	void AddLodDataToSkeletalMesh(const UInterchangeSkeletonFactoryNode* SkeletonFactoryNode, UInterchangeSkeletalMeshFactoryNode* SkeletalMeshFactoryNode, const TMap<int32, TArray<FString>>& SceneInstanceUidsPerLodIndex);
-
-	/** This function add one lod data node (the base lod) to the skeletal mesh. It is use when we do not bake the data, like for a scene import, or if a mesh is not reference by a scene node. */
-	void AddLodDataToSkeletalMesh(const UInterchangeSkeletonFactoryNode* SkeletonFactoryNode, UInterchangeSkeletalMeshFactoryNode* SkeletalMeshFactoryNode, const UInterchangeMeshNode* MeshNode);
+	/**
+	 * This function add all lod data node to the skeletal mesh.
+	 * @param NodeUidsPerLodIndex - The NodeUids can be a UInterchangeSceneNode or a UInterchangeMeshNode. The scene node can bake each instance of the mesh versus the mesh node will import only the modelled mesh.
+	 */
+	void AddLodDataToSkeletalMesh(const UInterchangeSkeletonFactoryNode* SkeletonFactoryNode, UInterchangeSkeletalMeshFactoryNode* SkeletalMeshFactoryNode, const TMap<int32, TArray<FString>>& NodeUidsPerLodIndex);
 
 	/** Translated scene nodes */
-	TArray<TObjectPtr<UInterchangeSceneNode>> SceneNodes;
+	TArray<UInterchangeSceneNode*> SceneNodes;
+
+	/* Meshes utilities, to parse the translated graph and extract the meshes informations. */
+	UInterchangePipelineMeshesUtilities* PipelineMeshesUtilities;
 };
 
 
diff --git a/Engine/Plugins/Runtime/Interchange/Source/Pipelines/Public/InterchangePipelineMeshesUtilities.h b/Engine/Plugins/Runtime/Interchange/Source/Pipelines/Public/InterchangePipelineMeshesUtilities.h
new file mode 100644
--- /dev/null
+++ b/Engine/Plugins/Runtime/Interchange/Source/Pipelines/Public/InterchangePipelineMeshesUtilities.h
@@ -0,0 +1,264 @@
+// Copyright Epic Games, Inc. All Rights Reserved.
+
+#pragma once
+
+#include "CoreMinimal.h"
+#include "UObject/Object.h"
+#include "UObject/ObjectMacros.h"
+
+#include "InterchangePipelineMeshesUtilities.generated.h"
+
+class UInterchangeBaseNodeContainer;
+class UInterchangeMeshNode;
+class UInterchangeSceneNode;
+
+/*
+* This container exist only because UPROPERTY cannot support nested container. See FInterchangeMeshInstance
+*/
+USTRUCT(BlueprintType)
+struct FInterchangeLodSceneNodeContainer
+{
+	GENERATED_BODY()
+
+	/**
+	* Each scene node here represent a mesh scene node. Only if we represent a lod group we can have more then 1 mesh scene node for a specific lod index.
+	*/
+	UPROPERTY(EditAnywhere, Category = "Interchange | Pipeline | MeshesInstance")
+	TArray<UInterchangeSceneNode*> SceneNodes;
+};
+
+/*
+* A mesh instance is a description of a translated scene node that point on a translated mesh asset.
+* A mesh instance pointing on a lod group can have many lods and many scene mesh nodes per lod index.
+* A mesh instance pointing on a mesh node will have only the lod 0 and will point on one scene mesh node.
+*/
+USTRUCT(BlueprintType)
+struct FInterchangeMeshInstance
+{
+	GENERATED_BODY()
+
+	FInterchangeMeshInstance()
+	{
+		LodGroupNode = nullptr;
+		bReferenceSkinnedMesh = false;
+	}
+	/**
+	* This ID represent either 1: a lod group scene node uid or 2: a mesh scene node uid.
+	*/
+	UPROPERTY(EditAnywhere, Category = "Interchange | Pipeline | MeshesInstance")
+	FString MeshInstanceUid;
+
+	/**
+	* If this mesh instance represent a LodGroup this member will not be null, but will be null if the mesh instance do not represent a lod group
+	*/
+	UPROPERTY(EditAnywhere, Category = "Interchange | Pipeline | MeshesInstance")
+	UInterchangeSceneNode* LodGroupNode;
+
+	UPROPERTY(EditAnywhere, Category = "Interchange | Pipeline | MeshesInstance")
+	bool bReferenceSkinnedMesh;
+
+	/**
+	* Each scene node here represent a mesh scene node. Only if we represent a lod group we can have more then 1 mesh scene node for a specific lod index.
+	*/
+	UPROPERTY(EditAnywhere, Category = "Interchange | Pipeline | MeshesInstance")
+	TMap<int32, FInterchangeLodSceneNodeContainer> SceneNodePerLodIndex;
+
+	/**
+	* All mesh geometry referenced by this MeshInstance.
+	*/
+	UPROPERTY(EditAnywhere, Category = "Interchange | Pipeline | MeshesGeometry")
+	TArray<FString> ReferencingMeshGeometryUids;
+};
+
+/*
+* A mesh geometry is a description of a translated mesh asset node that define a geometry.
+*/
+USTRUCT(BlueprintType)
+struct FInterchangeMeshGeometry
+{
+	GENERATED_BODY()
+
+	FInterchangeMeshGeometry()
+	{
+		MeshNode = nullptr;
+	}
+
+	/**
+	* Represent the unique id of the UInterchangeMeshNode represent by this structure.
+	*/
+	UPROPERTY(EditAnywhere, Category = "Interchange | Pipeline | MeshesGeometry")
+	FString MeshUid;
+
+	/**
+	* The UInterchangeMeshNode pointer represent by this structure.
+	*/
+	UPROPERTY(EditAnywhere, Category = "Interchange | Pipeline | MeshesGeometry")
+	UInterchangeMeshNode* MeshNode = nullptr;
+
+	/**
+	* All mesh instance referencing this UInterchangeMeshNode pointer.
+	*/
+	UPROPERTY(EditAnywhere, Category = "Interchange | Pipeline | MeshesGeometry")
+	TArray<FString> ReferencingMeshInstanceUids;
+};
+
+
+/**/
+UCLASS(BlueprintType)
+class INTERCHANGEPIPELINES_API UInterchangePipelineMeshesUtilities : public UObject
+{
+	GENERATED_BODY()
+public:
+	/**
+	* Create an instance of UInterchangePipelineMeshesUtilities.
+	*/
+	UFUNCTION(BlueprintCallable, Category = "Interchange | Pipeline | Meshes")
+	static UInterchangePipelineMeshesUtilities* CreateInterchangePipelineMeshesUtilities(UInterchangeBaseNodeContainer* BaseNodeContainer);
+
+	/**
+	* Get all mesh instance unique ids.
+	*/
+	UFUNCTION(BlueprintCallable, Category = "Interchange | Pipeline | Meshes")
+	void GetAllMeshInstanceUids(TArray<FString>& MeshInstanceUids) const;
+
+	/**
+	* Iterate all mesh instance.
+	*/
+	void IterateAllMeshInstance(TFunctionRef<void(const FInterchangeMeshInstance&)> IterationLambda) const;
+
+	/**
+	* Get all skinned mesh instance unique ids.
+	*/
+	UFUNCTION(BlueprintCallable, Category = "Interchange | Pipeline | Meshes")
+	void GetAllSkinnedMeshInstance(TArray<FString>& MeshInstanceUids) const;
+
+	/**
+	* Iterate all skinned mesh instance.
+	*/
+	void IterateAllSkinnedMeshInstance(TFunctionRef<void(const FInterchangeMeshInstance&)> IterationLambda) const;
+
+	/**
+	* Get all static mesh instance unique ids.
+	*/
+	UFUNCTION(BlueprintCallable, Category = "Interchange | Pipeline | Meshes")
+	void GetAllStaticMeshInstance(TArray<FString>& MeshInstanceUids) const;
+
+	/**
+	* Iterate all static mesh instance.
+	*/
+	void IterateAllStaticMeshInstance(TFunctionRef<void(const FInterchangeMeshInstance&)> IterationLambda) const;
+
+	/**
+	* Get all mesh geometry unique ids.
+	*/
+	UFUNCTION(BlueprintCallable, Category = "Interchange | Pipeline | Meshes")
+	void GetAllMeshGeometry(TArray<FString>& MeshGeometryUids) const;
+		
+	/**
+	* Iterate all mesh geometry.
+	*/
+	void IterateAllMeshGeometry(TFunctionRef<void(const FInterchangeMeshGeometry&)> IterationLambda) const;
+
+	/**
+	* Get all skinned mesh geometry unique ids.
+	*/
+	UFUNCTION(BlueprintCallable, Category = "Interchange | Pipeline | Meshes")
+	void GetAllSkinnedMeshGeometry(TArray<FString>& MeshGeometryUids) const;
+
+	/**
+	* Iterate all skinned mesh geometry.
+	*/
+	void IterateAllSkinnedMeshGeometry(TFunctionRef<void(const FInterchangeMeshGeometry&)> IterationLambda) const;
+
+	/**
+	* Get all static mesh geometry unique ids.
+	*/
+	UFUNCTION(BlueprintCallable, Category = "Interchange | Pipeline | Meshes")
+	void GetAllStaticMeshGeometry(TArray<FString>& MeshGeometryUids) const;
+
+	/**
+	* Iterate all static mesh geometry.
+	*/
+	void IterateAllStaticMeshGeometry(TFunctionRef<void(const FInterchangeMeshGeometry&)> IterationLambda) const;
+
+	/**
+	* Get all not instanced mesh geometry unique ids.
+	*/
+	UFUNCTION(BlueprintCallable, Category = "Interchange | Pipeline | Meshes")
+	void GetAllMeshGeometryNotInstanced(TArray<FString>& MeshGeometryUids) const;
+
+	/**
+	* Iterate all mesh geometry not instanced.
+	*/
+	void IterateAllMeshGeometryNotIntanced(TFunctionRef<void(const FInterchangeMeshGeometry&)> IterationLambda) const;
+
+	/**
+	* Return true if there is an existing FInterchangeMeshInstance matching the MeshInstanceUid key.
+	*/
+	UFUNCTION(BlueprintCallable, Category = "Interchange | Pipeline | Meshes")
+	bool IsValidMeshInstanceUid(const FString& MeshInstanceUid) const;
+
+	/**
+	* Get the instanced mesh from the unique ids.
+	*/
+	UFUNCTION(BlueprintCallable, Category = "Interchange | Pipeline | Meshes")
+	const FInterchangeMeshInstance& GetMeshInstanceByUid(const FString& MeshInstanceUid) const;
+
+	/**
+	* Return true if there is an existing FInterchangeMeshGeometry matching the MeshInstanceUid key.
+	*/
+	UFUNCTION(BlueprintCallable, Category = "Interchange | Pipeline | Meshes")
+	bool IsValidMeshGeometryUid(const FString& MeshGeometryUid) const;
+
+	/**
+	* Get the geometry mesh from the unique ids.
+	*/
+	UFUNCTION(BlueprintCallable, Category = "Interchange | Pipeline | Meshes")
+	const FInterchangeMeshGeometry& GetMeshGeometryByUid(const FString& MeshGeometryUid) const;
+
+	/**
+	* Get all instanced mesh uids using the mesh geometry unique ids.
+	*/
+	UFUNCTION(BlueprintCallable, Category = "Interchange | Pipeline | Meshes")
+	void GetAllMeshInstanceUidsUsingMeshGeometryUid(const FString& MeshGeometryUid, TArray<FString>& MeshInstanceUids) const;
+
+	/**
+	* Iterate all instanced mesh uids using the mesh geometry unique ids.
+	*/
+	void IterateAllMeshInstanceUsingMeshGeometry(const FString& MeshGeometryUid, TFunctionRef<void(const FInterchangeMeshInstance&)> IterationLambda) const;
+
+	/**
+	* Return a list of skinned FInterchangeMeshInstance uid that can be combined together.
+	* We cannot create a skinned mesh with multiple skeleton root node, This function return combined MeshInstance per skeleton roots
+	*/
+	void GetCombinedSkinnedMeshInstances(TMap<FString, TArray<FString>>& OutMeshInstanceUidsPerSkeletonRootUid) const;
+	
+	/**
+	* Return a list of skinned FInterchangeMeshGeometry uid that can be combined together.
+	* We cannot create a skinned mesh with multiple skeleton root node, This function return combined FInterchangeMeshGeometry per skeleton roots
+	*/
+	void GetCombinedSkinnedMeshGeometries(TMap<FString, TArray<FString>>& OutMeshGeometryUidsPerSkeletonRootUid) const;
+
+	/**
+	* Return the skeleton root node Uid, this is the uid for a UInterchangeSceneNode that has a "Joint" specialized type.
+	* Return an empty string if the MeshInstanceUid parameter point on nothing.
+	*/
+	UFUNCTION(BlueprintCallable, Category = "Interchange | Pipeline | Meshes")
+	FString GetMeshInstanceSkeletonRootUid(const FString& MeshInstanceUid) const;
+
+	FString GetMeshInstanceSkeletonRootUid(const FInterchangeMeshInstance& MeshInstance) const;
+
+	/**
+	* Return the skeleton root node Uid, this is the uid for a UInterchangeSceneNode that has a "Joint" specialized type.
+	* Return an empty string if the MeshGeometryUid parameter point on nothing.
+	*/
+	UFUNCTION(BlueprintCallable, Category = "Interchange | Pipeline | Meshes")
+	FString GetMeshGeometrySkeletonRootUid(const FString& MeshGeometryUid) const;
+
+	FString GetMeshGeometrySkeletonRootUid(const FInterchangeMeshGeometry& MeshGeometry) const;
+
+protected:
+	TMap<FString, FInterchangeMeshGeometry> MeshGeometriesPerMeshUid;
+	TMap<FString, FInterchangeMeshInstance> MeshInstancesPerMeshInstanceUid;
+	TMap<FString, FString> SkeletonRootUidPerMeshUid;
+};
diff --git a/Engine/Source/Runtime/Interchange/Core/Private/Nodes/InterchangeBaseNodeContainer.cpp b/Engine/Source/Runtime/Interchange/Core/Private/Nodes/InterchangeBaseNodeContainer.cpp
--- a/Engine/Source/Runtime/Interchange/Core/Private/Nodes/InterchangeBaseNodeContainer.cpp
+++ b/Engine/Source/Runtime/Interchange/Core/Private/Nodes/InterchangeBaseNodeContainer.cpp
@@ -57,7 +57,7 @@ bool UInterchangeBaseNodeContainer::IsNodeUidValid(const FString& NodeUniqueID)
 	return Nodes.Contains(NodeUniqueID);
 }
 
-void UInterchangeBaseNodeContainer::IterateNodes(TFunctionRef<void(const FString&, UInterchangeBaseNode*)> IterationLambda)
+void UInterchangeBaseNodeContainer::IterateNodes(TFunctionRef<void(const FString&, UInterchangeBaseNode*)> IterationLambda) const
 {
 	for (auto& NodeKeyValue : Nodes)
 	{
@@ -65,7 +65,7 @@ void UInterchangeBaseNodeContainer::IterateNodes(TFunctionRef<void(const FString
 	}
 }
 
-void UInterchangeBaseNodeContainer::GetRoots(TArray<FString>& RootNodes)
+void UInterchangeBaseNodeContainer::GetRoots(TArray<FString>& RootNodes) const
 {
 	for (auto& NodeKeyValue : Nodes)
 	{
@@ -76,7 +76,7 @@ void UInterchangeBaseNodeContainer::GetRoots(TArray<FString>& RootNodes)
 	}
 }
 
-void UInterchangeBaseNodeContainer::GetNodes(UClass* ClassNode, TArray<FString>& OutNodes)
+void UInterchangeBaseNodeContainer::GetNodes(const UClass* ClassNode, TArray<FString>& OutNodes) const
 {
 	OutNodes.Empty();
 	IterateNodes([&ClassNode, &OutNodes](const FString& NodeUid, UInterchangeBaseNode* Node)
diff --git a/Engine/Source/Runtime/Interchange/Core/Public/Nodes/InterchangeBaseNodeContainer.h b/Engine/Source/Runtime/Interchange/Core/Public/Nodes/InterchangeBaseNodeContainer.h
--- a/Engine/Source/Runtime/Interchange/Core/Public/Nodes/InterchangeBaseNodeContainer.h
+++ b/Engine/Source/Runtime/Interchange/Core/Public/Nodes/InterchangeBaseNodeContainer.h
@@ -41,15 +41,15 @@ public:
 	bool IsNodeUidValid(const FString& NodeUniqueID) const;
 
 	/** Unordered iteration of the all nodes */
-	void IterateNodes(TFunctionRef<void(const FString&, UInterchangeBaseNode*)> IterationLambda);
+	void IterateNodes(TFunctionRef<void(const FString&, UInterchangeBaseNode*)> IterationLambda) const;
 
 	/** Return all nodes that do not have any parent */
 	UFUNCTION(BlueprintCallable, Category = "Interchange | Node Container")
-	void GetRoots(TArray<FString>& RootNodes);
+	void GetRoots(TArray<FString>& RootNodes) const;
 
 	/** Return all nodes that are of the ClassNode type*/
 	UFUNCTION(BlueprintCallable, Category = "Interchange | Node Container")
-	void GetNodes(UClass* ClassNode, TArray<FString>& OutNodes);
+	void GetNodes(const UClass* ClassNode, TArray<FString>& OutNodes) const;
 
 	/** Get an node pointer */
 	UFUNCTION(BlueprintCallable, Category = "Interchange | Node Container")
