This change handle correctly all the bImportXXX options.

diff --git a/Engine/Plugins/Runtime/Interchange/Source/Pipelines/Private/InterchangeGenericAssetsPipeline.cpp b/Engine/Plugins/Runtime/Interchange/Source/Pipelines/Private/InterchangeGenericAssetsPipeline.cpp
--- a/Engine/Plugins/Runtime/Interchange/Source/Pipelines/Private/InterchangeGenericAssetsPipeline.cpp
+++ b/Engine/Plugins/Runtime/Interchange/Source/Pipelines/Private/InterchangeGenericAssetsPipeline.cpp
@@ -82,19 +82,26 @@ bool UInterchangeGenericAssetsPipeline::ExecutePreImportPipeline(UInterchangeBas
 		}
 	});
 	
-	//import textures
-	for (const UInterchangeTextureNode* TextureNode : TextureNodes)
+	if (bImportTextures)
 	{
-		CreateTexture2DFactoryNode(TextureNode);
+		//import textures
+		for (const UInterchangeTextureNode* TextureNode : TextureNodes)
+		{
+			CreateTexture2DFactoryNode(TextureNode);
+		}
 	}
-	//import materials
-	for (const UInterchangeMaterialNode* MaterialNode : MaterialNodes)
+
+	if (bImportMaterials)
 	{
-		CreateMaterialFactoryNode(MaterialNode);
+		//import materials
+		for (const UInterchangeMaterialNode* MaterialNode : MaterialNodes)
+		{
+			CreateMaterialFactoryNode(MaterialNode);
+		}
 	}
 
 
-	if (SkinnedMeshNodes.Num() > 0)
+	if (bImportSkeletalMeshes && SkinnedMeshNodes.Num() > 0)
 	{
 		auto SetSkeletalMeshDependencies = [&SkeletalMeshFactoryDependencyOrderPerSkeletonRootNodeUid](const FString& JointNodeUid, UInterchangeSkeletalMeshFactoryNode* SkeletalMeshFactoryNode)
 		{
@@ -255,6 +262,11 @@ bool UInterchangeGenericAssetsPipeline::ExecutePreImportPipeline(UInterchangeBas
 		}
 	}
 
+	if (bImportStaticMeshes)
+	{
+		//TODO Import staticmesh here
+	}
+
 	const UClass* SkeletalMeshFactoryNodeClass = UInterchangeSkeletalMeshFactoryNode::StaticClass();
 	TArray<FString> SkeletalMeshNodeUids;
 	BaseNodeContainer->GetNodes(SkeletalMeshFactoryNodeClass, SkeletalMeshNodeUids);
@@ -265,6 +277,7 @@ bool UInterchangeGenericAssetsPipeline::ExecutePreImportPipeline(UInterchangeBas
 
 	//TODO count also the imported animations
 
+	//If we import only one asset, and bUseSourceNameForAsset is true, we want to rename the asset using the file name.
 	const int32 MeshesAndAnimsImportedNodeCount = SkeletalMeshNodeUids.Num(); // + StaticMeshNodeUids.Num();
 	if (bUseSourceNameForAsset && MeshesAndAnimsImportedNodeCount == 1)
 	{
@@ -544,12 +557,15 @@ void UInterchangeGenericAssetsPipeline::AddLodDataToSkeletalMesh(const UIntercha
 			for (int32 MaterialIndex = 0; MaterialIndex < MaterialCount; ++MaterialIndex)
 			{
 				const FString MaterialFactoryNodeUid = UInterchangeMaterialFactoryNode::GetMaterialFactoryNodeUidFromMaterialNodeUid(MaterialDependencies[MaterialIndex]);
-				//Create a factory dependency so Material asset are import before the skeletal mesh asset
-				TArray<FString> FactoryDependencies;
-				SkeletalMeshFactoryNode->GetFactoryDependencies(FactoryDependencies);
-				if (!FactoryDependencies.Contains(MaterialFactoryNodeUid))
+				if (BaseNodeContainer->IsNodeUidValid(MaterialFactoryNodeUid))
 				{
-					SkeletalMeshFactoryNode->SetFactoryDependencyUid(MaterialFactoryNodeUid);
+					//Create a factory dependency so Material asset are import before the skeletal mesh asset
+					TArray<FString> FactoryDependencies;
+					SkeletalMeshFactoryNode->GetFactoryDependencies(FactoryDependencies);
+					if (!FactoryDependencies.Contains(MaterialFactoryNodeUid))
+					{
+						SkeletalMeshFactoryNode->SetFactoryDependencyUid(MaterialFactoryNodeUid);
+					}
 				}
 			}
 			LodDataNode->AddMeshUid(NodeUid);
