Fix CIS static analysis warning
#rnx

diff --git a/Engine/Plugins/Runtime/Interchange/Source/Pipelines/Private/InterchangePipelineMeshesUtilities.cpp b/Engine/Plugins/Runtime/Interchange/Source/Pipelines/Private/InterchangePipelineMeshesUtilities.cpp
--- a/Engine/Plugins/Runtime/Interchange/Source/Pipelines/Private/InterchangePipelineMeshesUtilities.cpp
+++ b/Engine/Plugins/Runtime/Interchange/Source/Pipelines/Private/InterchangePipelineMeshesUtilities.cpp
@@ -123,6 +123,10 @@ UInterchangePipelineMeshesUtilities* UInterchangePipelineMeshesUtilities::Create
 	for (const TPair<FString, FInterchangeMeshGeometry>& MeshGeometryUidAndMeshGeometry : PipelineMeshesUtilities->MeshGeometriesPerMeshUid)
 	{
 		const FInterchangeMeshGeometry& MeshGeometry = MeshGeometryUidAndMeshGeometry.Value;
+		if (!ensure(MeshGeometry.MeshNode))
+		{
+			continue;
+		}
 		if (!MeshGeometry.MeshNode->IsSkinnedMesh() || PipelineMeshesUtilities->SkeletonRootUidPerMeshUid.Contains(MeshGeometry.MeshUid))
 		{
 			continue;
