Adding cube map support to interchange.
Added cubemap support when importing a dds file into interchange.
Changed the default behavior when importing a dds cubemap into the editor to only preserve the existing mipmap when there is more than one mipmap.

diff --git a/Engine/Plugins/Runtime/Interchange/Source/Import/Private/Texture/InterchangeDDSTranslator.cpp b/Engine/Plugins/Runtime/Interchange/Source/Import/Private/Texture/InterchangeDDSTranslator.cpp
--- a/Engine/Plugins/Runtime/Interchange/Source/Import/Private/Texture/InterchangeDDSTranslator.cpp
+++ b/Engine/Plugins/Runtime/Interchange/Source/Import/Private/Texture/InterchangeDDSTranslator.cpp
@@ -7,6 +7,7 @@
 #include "IImageWrapper.h"
 #include "IImageWrapperModule.h"
 #include "InterchangeImportLog.h"
+#include "InterchangeTextureCubeNode.h"
 #include "InterchangeTextureNode.h"
 #include "Misc/ConfigCacheIni.h"
 #include "Misc/FileHelper.h"
@@ -16,9 +17,38 @@
 #include "Nodes/InterchangeBaseNodeContainer.h"
 #include "Texture/TextureTranslatorUtilities.h"
 
+#include "HAL/FileManager.h"
+#include "Serialization/Archive.h"
 
-//////////////////////////////////////////////////////////////////////////
-// DDS helper local function
+namespace UE::Interchange::Private::InterchangeDDSTranslator
+{ 
+	bool LoadDDSHeaderFromFile(TArray64<uint8>& OutHeader, const FString& Filename)
+	{
+		// The file is close when the archive is destroyed
+		TUniquePtr<FArchive> FileReaderArchive(IFileManager::Get().CreateFileReader(*Filename));
+
+		if (FileReaderArchive)
+		{
+			const int64 SizeOfFile = FileReaderArchive->TotalSize();
+			const int64 MinimalHeaderSize = FDDSLoadHelper::GetDDSHeaderMinimalSize();
+
+			// If the file is not bigger then the smallest header possible then clearly the file is not valid as a dds file.
+			if (SizeOfFile > MinimalHeaderSize)
+			{
+				const int64 MaximalHeaderSize = FDDSLoadHelper::GetDDSHeaderMaximalSize();
+				const int64 BytesToRead = SizeOfFile >= MaximalHeaderSize ? MaximalHeaderSize : MinimalHeaderSize;
+
+				OutHeader.Reset(BytesToRead);
+				OutHeader.AddUninitialized(BytesToRead);
+				FileReaderArchive->Serialize(OutHeader.GetData(), OutHeader.Num());
+
+				return true;
+			}
+		}
+
+		return false;
+	}
+}
 
 bool UInterchangeDDSTranslator::CanImportSourceData(const UInterchangeSourceData* InSourceData) const
 {
@@ -30,9 +60,9 @@ bool UInterchangeDDSTranslator::CanImportSourceData(const UInterchangeSourceData
 	}
 
 	/*
-	 * DDS file can also be a cube map so we have to open the file and see if its a valid 2D texture.
+	 * DDS file can also be a texture array so we have to open the file and see if its a valid 2D texture.
 	 */
-	TArray64<uint8> SourceDataBuffer;
+	TArray64<uint8> HeaderDataBuffer;
 	FString Filename = InSourceData->GetFilename();
 
 	if (!FPaths::FileExists(Filename))
@@ -40,24 +70,45 @@ bool UInterchangeDDSTranslator::CanImportSourceData(const UInterchangeSourceData
 		return false;
 	}
 
-	if (!FFileHelper::LoadFileToArray(SourceDataBuffer, *Filename))
+	if (!UE::Interchange::Private::InterchangeDDSTranslator::LoadDDSHeaderFromFile(HeaderDataBuffer, Filename))
 	{
 		return false;
 	}
 
-	const uint8* Buffer = SourceDataBuffer.GetData();
-	const uint8* BufferEnd = Buffer + SourceDataBuffer.Num();
-	// Validate it.
-	const int32 Length = BufferEnd - Buffer;
-
-
-	FDDSLoadHelper  DDSLoadHelper(Buffer, Length);
-	return DDSLoadHelper.IsValid2DTexture();
+	FDDSLoadHelper  DDSLoadHelper(HeaderDataBuffer.GetData(), HeaderDataBuffer.Num());
+	return DDSLoadHelper.IsValid2DTexture() || DDSLoadHelper.IsValidCubemapTexture();
 }
 
 bool UInterchangeDDSTranslator::Translate(const UInterchangeSourceData* SourceData, UInterchangeBaseNodeContainer& BaseNodeContainer) const
 {
-	return UE::Interchange::FTextureTranslatorUtilities::Generic2DTextureTranslate(SourceData, BaseNodeContainer);
+	/*
+	 * DDS file can also be a cube map so we have to open the file and see if its a valid 2D texture.
+	 */
+	FString Filename = SourceData->GetFilename();
+	if (!FPaths::FileExists(Filename))
+	{
+		return false;
+	}
+
+	TArray64<uint8> HeaderDataBuffer;
+	if (!UE::Interchange::Private::InterchangeDDSTranslator::LoadDDSHeaderFromFile(HeaderDataBuffer, Filename))
+	{
+		return false;
+	}
+
+
+	FDDSLoadHelper  DDSLoadHelper(HeaderDataBuffer.GetData(), HeaderDataBuffer.Num());
+	if (DDSLoadHelper.IsValid2DTexture())
+	{
+		return UE::Interchange::FTextureTranslatorUtilities::Generic2DTextureTranslate(SourceData, BaseNodeContainer);
+	}
+
+	if (DDSLoadHelper.IsValidCubemapTexture())
+	{
+		return UE::Interchange::FTextureTranslatorUtilities::GenericCubeTextureTranslate(SourceData, BaseNodeContainer);
+	}
+
+	return false;
 }
 
 TOptional<UE::Interchange::FImportImage> UInterchangeDDSTranslator::GetTexturePayloadData(const UInterchangeSourceData* SourceData, const FString& PayLoadKey) const
@@ -68,7 +119,6 @@ TOptional<UE::Interchange::FImportImage> UInterchangeDDSTranslator::GetTexturePa
 		return TOptional<UE::Interchange::FImportImage>();
 	}
 
-	TArray64<uint8> SourceDataBuffer;
 	FString Filename = SourceData->GetFilename();
 	
 	//Make sure the key fit the filename, The key should always be valid
@@ -78,37 +128,34 @@ TOptional<UE::Interchange::FImportImage> UInterchangeDDSTranslator::GetTexturePa
 		return TOptional<UE::Interchange::FImportImage>();
 	}
 
+	TArray64<uint8> DDSSourceData;
+
 	if (!FPaths::FileExists(Filename))
 	{
 		UE_LOG(LogInterchangeImport, Error, TEXT("Failed to import DDS, cannot open file. [%s]"), *Filename);
 		return TOptional<UE::Interchange::FImportImage>();
 	}
 
-	if (!FFileHelper::LoadFileToArray(SourceDataBuffer, *Filename))
+	if (!FFileHelper::LoadFileToArray(DDSSourceData, *Filename))
 	{
 		UE_LOG(LogInterchangeImport, Error, TEXT("Failed to import DDS, cannot load file content into an array. [%s]"), *Filename);
 		return TOptional<UE::Interchange::FImportImage>();
 	}
 
-	const uint8* Buffer = SourceDataBuffer.GetData();
-	const uint8* BufferEnd = Buffer + SourceDataBuffer.Num();
-
-	bool bAllowNonPowerOfTwo = false;
-	GConfig->GetBool(TEXT("TextureImporter"), TEXT("AllowNonPowerOfTwoTextures"), bAllowNonPowerOfTwo, GEditorIni);
-
-	// Validate it.
-	const int32 Length = BufferEnd - Buffer;
-
 
 	//
 	// DDS Texture
 	//
-	FDDSLoadHelper  DDSLoadHelper(Buffer, Length);
+	FDDSLoadHelper  DDSLoadHelper(DDSSourceData.GetData(), DDSSourceData.Num());
 	if (!DDSLoadHelper.IsValid2DTexture())
 	{
 		UE_LOG(LogInterchangeImport, Error, TEXT("Failed to import DDS, unsupported format. [%s]"), *Filename);
 		return TOptional<UE::Interchange::FImportImage>();
 	}
+
+	bool bAllowNonPowerOfTwo = false;
+	GConfig->GetBool(TEXT("TextureImporter"), TEXT("AllowNonPowerOfTwoTextures"), bAllowNonPowerOfTwo, GEditorIni);
+
 	// DDS 2d texture
 	if (!UE::Interchange::FImportImageHelper::IsImportResolutionValid(DDSLoadHelper.DDSHeader->dwWidth, DDSLoadHelper.DDSHeader->dwHeight, bAllowNonPowerOfTwo))
 	{
@@ -153,3 +200,124 @@ TOptional<UE::Interchange::FImportImage> UInterchangeDDSTranslator::GetTexturePa
 
 	return PayloadData;
 }
+
+TOptional<UE::Interchange::FImportSlicedImage> UInterchangeDDSTranslator::GetSlicedTexturePayloadData(const UInterchangeSourceData* SourceData, const FString& PayLoadKey) const
+{
+	if (!SourceData)
+	{
+		UE_LOG(LogInterchangeImport, Error, TEXT("Failed to import DDS, bad source data."));
+		return TOptional<UE::Interchange::FImportSlicedImage>();
+	}
+
+	FString Filename = SourceData->GetFilename();
+	
+	//Make sure the key fit the filename, The key should always be valid
+	if (!Filename.Equals(PayLoadKey))
+	{
+		UE_LOG(LogInterchangeImport, Error, TEXT("Failed to import DDS, wrong payload key. [%s]"), *Filename);
+		return TOptional<UE::Interchange::FImportSlicedImage>();
+	}
+
+	TArray64<uint8> DDSSourceData;
+	if (!FPaths::FileExists(Filename))
+	{
+		UE_LOG(LogInterchangeImport, Error, TEXT("Failed to import DDS, cannot open file. [%s]"), *Filename);
+		return TOptional<UE::Interchange::FImportSlicedImage>();
+	}
+
+	if (!FFileHelper::LoadFileToArray(DDSSourceData, *Filename))
+	{
+		UE_LOG(LogInterchangeImport, Error, TEXT("Failed to import DDS, cannot load file content into an array. [%s]"), *Filename);
+		return TOptional<UE::Interchange::FImportSlicedImage>();
+	}
+
+	//
+	// DDS Texture
+	//
+	FDDSLoadHelper DDSLoadHelper(DDSSourceData.GetData(), DDSSourceData.Num());
+
+	if (!DDSLoadHelper.IsValidCubemapTexture())
+	{
+		UE_LOG(LogInterchangeImport, Error, TEXT("Failed to import DDS, unsupported format. [%s]"), *Filename);
+		return TOptional<UE::Interchange::FImportSlicedImage>();
+	}
+
+	bool bAllowNonPowerOfTwo = false;
+	GConfig->GetBool(TEXT("TextureImporter"), TEXT("AllowNonPowerOfTwoTextures"), bAllowNonPowerOfTwo, GEditorIni);
+
+	if (!UE::Interchange::FImportImageHelper::IsImportResolutionValid(DDSLoadHelper.DDSHeader->dwWidth, DDSLoadHelper.DDSHeader->dwHeight, bAllowNonPowerOfTwo))
+	{
+		UE_LOG(LogInterchangeImport, Error, TEXT("Failed to import PCX, invalid resolution. Resolution[%d, %d], AllowPowerOfTwo[%s], [%s]"), DDSLoadHelper.DDSHeader->dwWidth, DDSLoadHelper.DDSHeader->dwHeight, bAllowNonPowerOfTwo ? TEXT("True") : TEXT("false"), *Filename);
+		return TOptional<UE::Interchange::FImportSlicedImage>();
+	}
+
+	ETextureSourceFormat SourceFormat = DDSLoadHelper.ComputeSourceFormat();
+
+	// Invalid DDS format
+	if (SourceFormat == TSF_Invalid)
+	{
+		UE_LOG(LogInterchangeImport, Error, TEXT("DDS file [%s] contains data in an unsupported format"), *Filename);
+		return TOptional<UE::Interchange::FImportSlicedImage>();
+	}
+
+	uint32 MipMapCount = DDSLoadHelper.ComputeMipMapCount();
+	if (MipMapCount <= 0)
+	{
+		UE_LOG(LogInterchangeImport, Error, TEXT("DDS file [%s] do not have any mipmap"), *Filename);
+		return TOptional<UE::Interchange::FImportSlicedImage>();
+	}
+
+	UE::Interchange::FImportSlicedImage PayloadData;
+	PayloadData.Init(
+		DDSLoadHelper.GetSizeX(),
+		DDSLoadHelper.GetSizeY(),
+		DDSLoadHelper.GetSliceCount(),
+		MipMapCount,
+		SourceFormat,
+		!FTextureSource::IsHDR(SourceFormat)
+	);
+
+	uint8* DestMipData[MAX_TEXTURE_MIP_COUNT] = { 0 };
+	int64 MipsSize[MAX_TEXTURE_MIP_COUNT] = { 0 };
+	int64 NumBytes = 0;
+
+	for (uint32 MipIndex = 0; MipIndex < MipMapCount; ++MipIndex)
+	{
+		int64 MipSize = PayloadData.GetMipSize(MipIndex);
+		MipsSize[MipIndex] = MipSize;
+		// Six time since we have six slice
+		NumBytes += MipSize * 6;
+	}
+
+	PayloadData.RawData.AddUninitialized(NumBytes);
+
+	for (uint32 MipIndex = 0; MipIndex < MipMapCount; ++MipIndex)
+	{
+		DestMipData[MipIndex] = PayloadData.GetMipData(MipIndex);
+	}
+
+	for (uint32 SliceIndex = 0; SliceIndex < DDSLoadHelper.GetSliceCount(); ++SliceIndex)
+	{
+		const uint8* SrcMipData = DDSLoadHelper.GetDDSDataPointer((ECubeFace)SliceIndex);
+		for (uint32 MipIndex = 0; MipIndex < MipMapCount; ++MipIndex)
+		{
+			FMemory::Memcpy(DestMipData[MipIndex] + MipsSize[MipIndex] * SliceIndex, SrcMipData, MipsSize[MipIndex]);
+			SrcMipData += MipsSize[MipIndex];
+		}
+	}
+
+	if (MipMapCount > 1)
+	{
+		PayloadData.MipGenSettings = TMGS_LeaveExistingMips;
+	}
+
+	if (FTextureSource::IsHDR(SourceFormat))
+	{
+		// the loader can suggest a compression setting
+		PayloadData.CompressionSettings = TC_HDR;
+	}
+
+	return PayloadData;
+}
+
+
diff --git a/Engine/Plugins/Runtime/Interchange/Source/Import/Private/Texture/InterchangeSlicedTexturePayloadData.cpp b/Engine/Plugins/Runtime/Interchange/Source/Import/Private/Texture/InterchangeSlicedTexturePayloadData.cpp
new file mode 100644
--- /dev/null
+++ b/Engine/Plugins/Runtime/Interchange/Source/Import/Private/Texture/InterchangeSlicedTexturePayloadData.cpp
@@ -0,0 +1,50 @@
+// Copyright Epic Games, Inc. All Rights Reserved.
+
+#include "Texture/InterchangeSlicedTexturePayloadData.h"
+
+#include "CoreMinimal.h"
+#include "Engine/Texture.h"
+
+
+void UE::Interchange::FImportSlicedImage::Init(int32 InSizeX, int32 InSizeY, int32 InNumSlice, int32 InNumMips, ETextureSourceFormat InFormat, bool InSRGB)
+{
+	SizeX = InSizeX;
+	SizeY = InSizeY;
+	NumSlice = InNumSlice;
+	NumMips = InNumMips;
+	Format = InFormat;
+	bSRGB = InSRGB;
+}
+
+int64 UE::Interchange::FImportSlicedImage::GetMipSize(int32 InMipIndex) const
+{
+	check(InMipIndex >= 0);
+	check(InMipIndex < NumMips);
+	const int32 MipSizeX = FMath::Max(SizeX >> InMipIndex, 1);
+	const int32 MipSizeY = FMath::Max(SizeY >> InMipIndex, 1);
+	return (int64)MipSizeX * MipSizeY * FTextureSource::GetBytesPerPixel(Format);
+
+}
+
+const uint8* UE::Interchange::FImportSlicedImage::GetMipData(int32 InMipIndex, int32 InSliceIndex) const
+{
+	int64 Offset = 0;
+	int32 MipIndex;
+	for (MipIndex = 0; MipIndex < InMipIndex; ++MipIndex)
+	{
+		Offset += GetMipSize(MipIndex) * NumSlice;
+	}
+
+	if (InSliceIndex != INDEX_NONE)
+	{ 
+		Offset += GetMipSize(MipIndex) * InSliceIndex;
+	}
+
+	return &RawData[Offset];
+}
+
+uint8* UE::Interchange::FImportSlicedImage::GetMipData(int32 InMipIndex, int32 InSliceIndex)
+{
+	const uint8* ConstBufferPtr = static_cast<const FImportSlicedImage*>(this)->GetMipData(InMipIndex, InSliceIndex);
+	return const_cast<uint8*>(ConstBufferPtr);
+}
diff --git a/Engine/Plugins/Runtime/Interchange/Source/Import/Private/Texture/InterchangeTextureFactory.cpp b/Engine/Plugins/Runtime/Interchange/Source/Import/Private/Texture/InterchangeTextureFactory.cpp
--- a/Engine/Plugins/Runtime/Interchange/Source/Import/Private/Texture/InterchangeTextureFactory.cpp
+++ b/Engine/Plugins/Runtime/Interchange/Source/Import/Private/Texture/InterchangeTextureFactory.cpp
@@ -1,19 +1,25 @@
-// Copyright Epic Games, Inc. All Rights Reserved. 
+// Copyright Epic Games, Inc. All Rights Reserved.
+
 #include "Texture/InterchangeTextureFactory.h"
 
 #include "Engine/Texture.h"
 #include "Engine/Texture2D.h"
+#include "Engine/TextureCube.h"
 #include "InterchangeAssetImportData.h"
 #include "InterchangeImportCommon.h"
 #include "InterchangeImportLog.h"
+#include "InterchangeTextureCubeFactoryNode.h"
+#include "InterchangeTextureCubeNode.h"
 #include "InterchangeTextureFactoryNode.h"
 #include "InterchangeTextureNode.h"
 #include "InterchangeTranslatorBase.h"
+#include "Misc/TVariant.h"
 #include "Nodes/InterchangeBaseNode.h"
 #include "Texture/InterchangeBlockedTexturePayloadInterface.h"
+#include "Texture/InterchangeSlicedTexturePayloadInterface.h"
 #include "Texture/InterchangeTexturePayloadInterface.h"
 #include "TextureCompiler.h"
-
+#include "UObject/ObjectMacros.h"
 
 #if WITH_EDITORONLY_DATA
 
@@ -21,6 +27,314 @@
 
 #endif //WITH_EDITORONLY_DATA
 
+
+
+namespace UE::Interchange::Private::InterchangeTextureFactory
+{
+	/**
+	 * Return the supported class if the node is one otherwise return nullptr
+	 */
+	UClass* GetSupportedFactoryNodeClass(const UInterchangeBaseNode* AssetNode)
+	{
+		UClass* TextureCubeFactoryClass = UInterchangeTextureCubeFactoryNode::StaticClass();
+		UClass* TextureFactoryClass = UInterchangeTextureFactoryNode::StaticClass();
+		UClass* AssetClass = AssetNode->GetClass();
+#if USTRUCT_FAST_ISCHILDOF_IMPL == USTRUCT_ISCHILDOF_STRUCTARRAY
+		if(AssetClass->IsChildOf(TextureFactoryClass))
+		{
+			return TextureFactoryClass;
+		}
+		else if(AssetClass->IsChildOf(TextureCubeFactoryClass))
+		{
+			return TextureCubeFactoryClass;
+		}
+#else
+		while (AssetClass)
+		{
+			if (AssetClass == TextureCubeFactoryClass || AssetClass == TextureFactoryClass)
+			{
+				return AssetClass;
+			}
+
+			AssetClass = AssetClass->GetSuperClass();
+		}
+#endif
+
+		return nullptr;
+	}
+
+	using FTextureFactoryNodeVariant = TVariant<FEmptyVariantState, UInterchangeTextureFactoryNode*, UInterchangeTextureCubeFactoryNode*>; 
+
+	FTextureFactoryNodeVariant GetAsTextureFactoryNodeVariant(UInterchangeBaseNode* AssetNode, UClass* SupportedFactoryNodeClass)
+	{
+		if (AssetNode)
+		{
+			if (!SupportedFactoryNodeClass)
+			{
+				SupportedFactoryNodeClass = GetSupportedFactoryNodeClass(AssetNode);
+			}
+
+			if (SupportedFactoryNodeClass == UInterchangeTextureFactoryNode::StaticClass())
+			{
+				return FTextureFactoryNodeVariant(TInPlaceType<UInterchangeTextureFactoryNode*>(), static_cast<UInterchangeTextureFactoryNode*>(AssetNode));
+			}
+
+			if (SupportedFactoryNodeClass == UInterchangeTextureCubeFactoryNode::StaticClass())
+			{
+				return FTextureFactoryNodeVariant(TInPlaceType<UInterchangeTextureCubeFactoryNode*>(), static_cast<UInterchangeTextureCubeFactoryNode*>(AssetNode));
+			}
+		}
+
+		return {};
+	}
+
+	using FTextureNodeVariant = TVariant<FEmptyVariantState, const UInterchangeTextureNode*, const UInterchangeTextureCubeNode*>; 
+
+	FTextureNodeVariant GetTextureNodeVariantFromFactoryVariant(const FTextureFactoryNodeVariant& FactoryVariant, const UInterchangeBaseNodeContainer* NodeContainer)
+	{
+		FString TextureNodeUniqueID;
+
+		if (UInterchangeTextureFactoryNode* const* TextureFactoryNode = FactoryVariant.TryGet<UInterchangeTextureFactoryNode*>())
+		{
+			if ((*TextureFactoryNode)->GetCustomTranslatedTextureNodeUid(TextureNodeUniqueID))
+			{
+				if (const UInterchangeTextureNode* TextureTranslatedNode = Cast<UInterchangeTextureNode>(NodeContainer->GetNode(TextureNodeUniqueID)))
+				{
+					return FTextureNodeVariant(TInPlaceType<const UInterchangeTextureNode*>(), TextureTranslatedNode);
+				}
+			}
+		}
+		else if (UInterchangeTextureCubeFactoryNode* const* TextureCubeFactoryNode = FactoryVariant.TryGet<UInterchangeTextureCubeFactoryNode*>())
+		{
+			if ((*TextureCubeFactoryNode)->GetCustomTranslatedTextureNodeUid(TextureNodeUniqueID))
+			{
+				if (const UInterchangeTextureCubeNode* TextureCubeTranslatedNode = Cast<UInterchangeTextureCubeNode>(NodeContainer->GetNode(TextureNodeUniqueID)))
+				{
+					return FTextureNodeVariant(TInPlaceType<const UInterchangeTextureCubeNode*>(), TextureCubeTranslatedNode);
+				}
+			}
+		}
+
+		return {};
+	}
+
+	bool HasPayloadKey(const FTextureNodeVariant& TextureNodeVariant)
+	{
+		if (const UInterchangeTextureNode* const* TextureNode =  TextureNodeVariant.TryGet<const UInterchangeTextureNode*>())
+		{
+			return (*TextureNode)->GetPayLoadKey().IsSet();
+		}
+
+		if (const UInterchangeTextureCubeNode* const* TextureCubeNode = TextureNodeVariant.TryGet<const UInterchangeTextureCubeNode*>())
+		{
+			return (*TextureCubeNode)->GetPayLoadKey().IsSet();
+		}
+
+	return false;
+	}
+
+	TOptional<FString> GetPayloadKey(const FTextureNodeVariant& TextureNodeVariant)
+	{
+		if (const UInterchangeTextureNode* const* TextureNode =  TextureNodeVariant.TryGet<const UInterchangeTextureNode*>())
+		{
+			return (*TextureNode)->GetPayLoadKey();
+		}
+
+		if (const UInterchangeTextureCubeNode* const* TextureCubeNode = TextureNodeVariant.TryGet<const UInterchangeTextureCubeNode*>())
+		{
+			return (*TextureCubeNode)->GetPayLoadKey();
+		}
+
+		return {};
+	}
+
+	using FTexturePayloadVariant = TVariant<FEmptyVariantState, TOptional<FImportImage>,  TOptional<FImportBlockedImage>,  TOptional<FImportSlicedImage>>;
+
+	FTexturePayloadVariant GetTexturePayload(const UInterchangeSourceData* SourceData, const FString& PayloadKey, const FTextureNodeVariant& TextureNodeVariant, const UInterchangeTranslatorBase* Translator)
+	{
+		// Standard texture 2D payload
+		if (const UInterchangeTextureNode* const* TextureNode =  TextureNodeVariant.TryGet<const UInterchangeTextureNode*>())
+		{
+			if (const IInterchangeTexturePayloadInterface* TextureTranslator = Cast<IInterchangeTexturePayloadInterface>(Translator))
+			{
+				return FTexturePayloadVariant(TInPlaceType<TOptional<FImportImage>>(), TextureTranslator->GetTexturePayloadData(SourceData, PayloadKey));
+			}
+			else if (const IInterchangeBlockedTexturePayloadInterface* BlockedTextureTranslator = Cast<IInterchangeBlockedTexturePayloadInterface>(Translator))
+			{
+				return FTexturePayloadVariant(TInPlaceType<TOptional<FImportBlockedImage>>(), BlockedTextureTranslator->GetBlockedTexturePayloadData((*TextureNode)->GetSourceBlocks(), SourceData));
+			}
+		}
+
+		// Cube texture payload
+		if (TextureNodeVariant.IsType<const UInterchangeTextureCubeNode*>())
+		{
+			if (const IInterchangeSlicedTexturePayloadInterface* SlicedTextureTranslator = Cast<IInterchangeSlicedTexturePayloadInterface>(Translator))
+			{
+				return FTexturePayloadVariant(TInPlaceType<TOptional<FImportSlicedImage>>(), SlicedTextureTranslator->GetSlicedTexturePayloadData(SourceData, PayloadKey));
+			}
+		}
+
+		return {};
+	}
+
+#if WITH_EDITORONLY_DATA
+	void SetupTexture2DSourceData(UTexture2D* Texture2D, const FImportImage& Image)
+	{
+		Texture2D->Source.Init(
+			Image.SizeX,
+			Image.SizeY,
+			/*NumSlices=*/ 1,
+			Image.NumMips,
+			Image.Format,
+			Image.RawData.GetData()
+		);
+
+		Texture2D->CompressionSettings = Image.CompressionSettings;
+		Texture2D->SRGB = Image.SRGB;
+
+		//If the MipGenSettings was set by the translator, we must apply it before the build
+		if (Image.MipGenSettings.IsSet())
+		{
+			// if the source has mips we keep the mips by default, unless the user changes that
+			Texture2D->MipGenSettings = Image.MipGenSettings.GetValue();
+		}
+	}
+
+	bool SetupTexture2DSourceData(UTexture2D* Texture2D, const FImportBlockedImage& BlockedImage)
+	{
+		if (BlockedImage.IsBlockedData())
+		{
+			if (BlockedImage.HasData())
+			{
+				ETextureSourceFormat Format = BlockedImage.ImagesData[0].Format;
+				TextureCompressionSettings CompressionSettings = BlockedImage.ImagesData[0].CompressionSettings;
+				bool bSRGB = BlockedImage.ImagesData[0].SRGB;
+
+				TArray<const uint8*> SourceImageDatasPtr;
+				SourceImageDatasPtr.Reserve(BlockedImage.ImagesData.Num());
+				for (const  UE::Interchange::FImportImage& Image : BlockedImage.ImagesData)
+				{
+					SourceImageDatasPtr.Add(Image.RawData.GetData());
+				}
+
+				Texture2D->Source.InitBlocked(
+					&Format,
+					BlockedImage.BlocksData.GetData(),
+					/*NumSlices=*/ 1,
+					BlockedImage.BlocksData.Num(),
+					SourceImageDatasPtr.GetData()
+				);
+
+				Texture2D->CompressionSettings = CompressionSettings;
+				Texture2D->SRGB = bSRGB;
+				Texture2D->VirtualTextureStreaming = true;
+
+				if (BlockedImage.ImagesData[0].MipGenSettings.IsSet())
+				{
+					// if the source has mips we keep the mips by default, unless the user changes that
+					Texture2D->MipGenSettings = BlockedImage.ImagesData[0].MipGenSettings.GetValue();
+				}
+
+				return true;
+			}
+		}
+		else
+		{
+			//Import as a normal texture
+			if (BlockedImage.ImagesData.Num() == 1)
+			{
+				SetupTexture2DSourceData(Texture2D, BlockedImage.ImagesData[0]);
+				return true;
+			}
+		}
+
+		return false;
+	}
+
+	bool SetupTexture2DSourceData(UTexture2D* Texture2D, const FTexturePayloadVariant& TexturePayload)
+	{
+		if (const TOptional<FImportBlockedImage>* BlockedImage = TexturePayload.TryGet<TOptional<FImportBlockedImage>>())
+		{
+			if (BlockedImage->IsSet())
+			{
+				return SetupTexture2DSourceData(Texture2D, BlockedImage->GetValue());
+			}
+		}
+		else if (const TOptional<FImportImage>* Image = TexturePayload.TryGet<TOptional<FImportImage>>())
+		{
+			if (Image->IsSet())
+			{
+				SetupTexture2DSourceData(Texture2D, Image->GetValue());
+				return true;
+			}
+		}
+
+		return false;
+	}
+
+	bool SetupTextureCubeSourceData(UTextureCube* TextureCube, const FTexturePayloadVariant& TexturePayload)
+	{
+		if (const TOptional<FImportSlicedImage>* OptionalSlicedImage = TexturePayload.TryGet<TOptional<FImportSlicedImage>>())
+		{
+			if (OptionalSlicedImage->IsSet())
+			{
+				const FImportSlicedImage& SlicedImage = OptionalSlicedImage->GetValue();
+
+				// Cube texture always have six slice
+				if (SlicedImage.NumSlice == 6)
+				{
+					TextureCube->Source.Init(
+						SlicedImage.SizeX,
+						SlicedImage.SizeY,
+						SlicedImage.NumSlice,
+						SlicedImage.NumMips,
+						SlicedImage.Format
+					);
+
+					TextureCube->CompressionSettings = SlicedImage.CompressionSettings;
+					TextureCube->SRGB = SlicedImage.bSRGB;
+
+
+					uint8* DestMipData[MAX_TEXTURE_MIP_COUNT] = {0};
+					int32 MipSize[MAX_TEXTURE_MIP_COUNT] = {0};
+					for (int32 MipIndex = 0; MipIndex < SlicedImage.NumMips; ++MipIndex)
+					{
+						DestMipData[MipIndex] = TextureCube->Source.LockMip(MipIndex);
+						MipSize[MipIndex] = TextureCube->Source.CalcMipSize(MipIndex) / 6;
+					}
+
+					for (int32 SliceIndex = 0; SliceIndex < 6; ++SliceIndex)
+					{
+						const uint8* SrcMipData = SlicedImage.GetMipData(0, SliceIndex);
+						for (int32 MipIndex = 0; MipIndex < SlicedImage.NumMips; ++MipIndex)
+						{
+							FMemory::Memcpy(DestMipData[MipIndex] + MipSize[MipIndex] * SliceIndex, SrcMipData, MipSize[MipIndex]);
+							SrcMipData += MipSize[MipIndex];
+						}
+					}
+
+					for (int32 MipIndex = 0; MipIndex < SlicedImage.NumMips; ++MipIndex)
+					{
+						TextureCube->Source.UnlockMip(MipIndex);
+					}
+
+					if (SlicedImage.MipGenSettings.IsSet())
+					{
+						// if the source has mips we keep the mips by default, unless the user changes that
+						TextureCube->MipGenSettings = SlicedImage.MipGenSettings.GetValue();
+					}
+
+					return true;
+				}
+			}
+		}
+
+		return false;
+	}
+#endif // WITH_EDITORONLY_DATA
+ }
+
+
 UClass* UInterchangeTextureFactory::GetFactoryClass() const
 {
 	return UTexture::StaticClass();
@@ -28,6 +342,8 @@ UClass* UInterchangeTextureFactory::GetFactoryClass() const
 
 UObject* UInterchangeTextureFactory::CreateEmptyAsset(const FCreateAssetParams& Arguments) const
 {
+	using namespace  UE::Interchange::Private::InterchangeTextureFactory;
+
 	UObject* Texture = nullptr;
 
 #if WITH_EDITORONLY_DATA
@@ -37,43 +353,33 @@ UObject* UInterchangeTextureFactory::CreateEmptyAsset(const FCreateAssetParams&
 		return nullptr;
 	}
 
-	if (!Arguments.AssetNode->GetAssetClass()->IsChildOf(UTexture::StaticClass()))
+	const UClass* TextureClass = Arguments.AssetNode->GetAssetClass();
+	if (!TextureClass || !TextureClass->IsChildOf(UTexture::StaticClass()))
 	{
 		UE_LOG(LogInterchangeImport, Error, TEXT("UInterchangeTextureFactory: Asset node parameter class doesnt derive from UTexture."));
 		return nullptr;
 	}
 
-	UInterchangeTextureFactoryNode* TextureFactoryNode = Cast<UInterchangeTextureFactoryNode>(Arguments.AssetNode);
-	if (TextureFactoryNode == nullptr)
+	UClass* SupportedFactoryNodeClass = GetSupportedFactoryNodeClass(Arguments.AssetNode);
+	if (SupportedFactoryNodeClass == nullptr)
 	{
-		UE_LOG(LogInterchangeImport, Error, TEXT("UInterchangeTextureFactory: Asset node parameter is not a UInterchangeTextureFactoryNode."));
+		UE_LOG(LogInterchangeImport, Error, TEXT("UInterchangeTextureFactory: Asset node parameter is not a UInterchangeTextureFactoryNode or UInterchangeTextureCubeFactoryNode."));
 		return nullptr;
 	}
 
-	FString TextureNodeUniqueID;
-	const UInterchangeTextureNode* TextureTranslatedNode = nullptr;
-	if (TextureFactoryNode->GetCustomTranslatedTextureNodeUid(TextureNodeUniqueID))
-	{
-		TextureTranslatedNode = Cast<UInterchangeTextureNode>(Arguments.NodeContainer->GetNode(TextureNodeUniqueID));
-	}
 
-	if (!TextureTranslatedNode)
+	FTextureNodeVariant TextureNodeVariant = GetTextureNodeVariantFromFactoryVariant(GetAsTextureFactoryNodeVariant(Arguments.AssetNode, SupportedFactoryNodeClass), Arguments.NodeContainer);
+	if (TextureNodeVariant.IsType<FEmptyVariantState>())
 	{
-		UE_LOG(LogInterchangeImport, Error, TEXT("UInterchangeTextureFactory: Asset factory node (UInterchangeTextureFactoryNode) do not reference a valid UInterchangeTextureNode translated node."));
-		return nullptr;
-	}
-
-	const TOptional<FString>& PayLoadKey = TextureTranslatedNode->GetPayLoadKey();
-	if (!PayLoadKey.IsSet())
-	{
-		UE_LOG(LogInterchangeImport, Error, TEXT("UInterchangeTextureFactory: Texture translated node (UInterchangeTextureNode) doesnt have a payload key."));
+		UE_LOG(LogInterchangeImport, Error, TEXT("UInterchangeTextureFactory: Asset factory node (%s) do not reference a valid texture translated node.")
+				, *SupportedFactoryNodeClass->GetAuthoredName()
+			);
 		return nullptr;
 	}
 
-	const UClass* TextureClass = TextureFactoryNode->GetAssetClass();
-	if (!TextureClass || !TextureClass->IsChildOf(GetFactoryClass()))
+	if (!HasPayloadKey(TextureNodeVariant))
 	{
-		UE_LOG(LogInterchangeImport, Error, TEXT("UInterchangeTextureFactory: Texture factory node asset class doesnt match Texture factory supported class."));
+		UE_LOG(LogInterchangeImport, Error, TEXT("UInterchangeTextureFactory: Texture translated node doesnt have a payload key."));
 		return nullptr;
 	}
 
@@ -112,89 +418,69 @@ UObject* UInterchangeTextureFactory::CreateAsset(const UInterchangeTextureFactor
 
 #else
 
+	using namespace  UE::Interchange::Private::InterchangeTextureFactory;
+
 	if (!Arguments.AssetNode)
 	{
 		UE_LOG(LogInterchangeImport, Error, TEXT("UInterchangeTextureFactory: Asset node parameter is null."));
 		return nullptr;
 	}
 
-	if (!Arguments.AssetNode->GetAssetClass()->IsChildOf(UTexture::StaticClass()))
+	const UClass* TextureClass = Arguments.AssetNode->GetAssetClass();
+	if (!TextureClass || !TextureClass->IsChildOf(UTexture::StaticClass()))
 	{
 		UE_LOG(LogInterchangeImport, Error, TEXT("UInterchangeTextureFactory: Asset node parameter class doesnt derive from UTexture."));
 		return nullptr;
 	}
 
-	UInterchangeTextureFactoryNode* TextureFactoryNode = Cast<UInterchangeTextureFactoryNode>(Arguments.AssetNode);
-	if (TextureFactoryNode == nullptr)
+	UClass* SupportedFactoryNodeClass = GetSupportedFactoryNodeClass(Arguments.AssetNode);
+	if (SupportedFactoryNodeClass == nullptr)
 	{
-		UE_LOG(LogInterchangeImport, Error, TEXT("UInterchangeTextureFactory: Asset node parameter is not a UInterchangeTextureFactoryNode."));
+		UE_LOG(LogInterchangeImport, Error, TEXT("UInterchangeTextureFactory: Asset node parameter is not a UInterchangeTextureFactoryNode or UInterchangeTextureCubeFactoryNode."));
 		return nullptr;
 	}
 
-	FString TextureNodeUniqueID;
-	const UInterchangeTextureNode* TextureTranslatedNode = nullptr;
-	if (TextureFactoryNode->GetCustomTranslatedTextureNodeUid(TextureNodeUniqueID))
+	FTextureFactoryNodeVariant TextureFactoryNodeVariant = GetAsTextureFactoryNodeVariant(Arguments.AssetNode, SupportedFactoryNodeClass);
+	FTextureNodeVariant TextureNodeVariant = GetTextureNodeVariantFromFactoryVariant(TextureFactoryNodeVariant, Arguments.NodeContainer);
+	if (TextureNodeVariant.IsType<FEmptyVariantState>())
 	{
-		TextureTranslatedNode = Cast<UInterchangeTextureNode>(Arguments.NodeContainer->GetNode(TextureNodeUniqueID));
-	}
-
-	if (!TextureTranslatedNode)
-	{
-		UE_LOG(LogInterchangeImport, Error, TEXT("UInterchangeTextureFactory: Asset factory node (UInterchangeTextureFactoryNode) do not reference a valid UInterchangeTextureNode translated node."));
+		UE_LOG(LogInterchangeImport, Error, TEXT("UInterchangeTextureFactory: Asset factory node (%s) do not reference a valid texture translated node.")
+				, *SupportedFactoryNodeClass->GetAuthoredName()
+			);
 		return nullptr;
 	}
 
-	const TOptional<FString>& PayLoadKey = TextureTranslatedNode->GetPayLoadKey();
+	const TOptional<FString>& PayLoadKey = GetPayloadKey(TextureNodeVariant);
 	if (!PayLoadKey.IsSet())
 	{
 		UE_LOG(LogInterchangeImport, Error, TEXT("UInterchangeTextureFactory: Texture translated node (UInterchangeTextureNode) doesnt have a payload key."));
 		return nullptr;
 	}
 
-	TOptional<UE::Interchange::FImportImage> PayloadData;
-	TOptional<UE::Interchange::FImportBlockedImage> BlockedPayloadData;
-	if (const IInterchangeTexturePayloadInterface* TextureTranslator = Cast<IInterchangeTexturePayloadInterface>(Arguments.Translator))
-	{
-		 PayloadData = TextureTranslator->GetTexturePayloadData(Arguments.SourceData, PayLoadKey.GetValue());
-	}
-	else if (const IInterchangeBlockedTexturePayloadInterface* BlockedTextureTranslator = Cast<IInterchangeBlockedTexturePayloadInterface>(Arguments.Translator))
-	{
-		/**
-		 * Possible improvement store when a file was last modified and use that as heuristic to skip the file assuming that the file is same as the one in the previous import
-		 * (how to deal with engine changes or translator changes/hotfix?)
-		 */ 
-		BlockedPayloadData = BlockedTextureTranslator->GetBlockedTexturePayloadData(TextureTranslatedNode->GetSourceBlocks(), Arguments.SourceData);
-	}
+	FTexturePayloadVariant TexturePayload = GetTexturePayload(Arguments.SourceData, PayLoadKey.GetValue(), TextureNodeVariant, Arguments.Translator);
 
-	if(!PayloadData.IsSet() && !BlockedPayloadData.IsSet())
+	if(TexturePayload.IsType<FEmptyVariantState>())
 	{
-		UE_LOG(LogInterchangeImport, Error, TEXT("UInterchangeTextureFactory: Invalid payload."));
-		return nullptr;
-	}
-
-	const UClass* TextureClass = TextureFactoryNode->GetAssetClass();
-	if (!TextureClass || !TextureClass->IsChildOf(GetFactoryClass()))
-	{
-		UE_LOG(LogInterchangeImport, Error, TEXT("UInterchangeTextureFactory: Texture factory node asset class doesnt match Texture factory supported class."));
+		UE_LOG(LogInterchangeImport, Error, TEXT("UInterchangeTextureFactory: Invalid translator couldn't retrive a payload."));
 		return nullptr;
 	}
 
 	// create an asset if it doesn't exist
 	UObject* ExistingAsset = StaticFindObject(nullptr, Arguments.Parent, *Arguments.AssetName);
 
-	UObject* Texture = nullptr;
+	UTexture* Texture = nullptr;
 	// create a new texture or overwrite existing asset, if possible
 	if (!ExistingAsset)
 	{
 		//NewObject is not thread safe, the asset registry directory watcher tick on the main thread can trig before we finish initializing the UObject and will crash
 		//The UObject should have been create by calling CreateEmptyAsset on the main thread.
 		check(IsInGameThread());
-		Texture = NewObject<UObject>(Arguments.Parent, TextureClass, *Arguments.AssetName, RF_Public | RF_Standalone);
+		Texture = NewObject<UTexture>(Arguments.Parent, TextureClass, *Arguments.AssetName, RF_Public | RF_Standalone);
 	}
 	else if(ExistingAsset->GetClass()->IsChildOf(TextureClass))
 	{
 		//This is a reimport, we are just re-updating the source data
-		Texture = ExistingAsset;
+		Texture = static_cast<UTexture*>(ExistingAsset);
 	}
 
 	if (!Texture)
@@ -203,107 +489,60 @@ UObject* UInterchangeTextureFactory::CreateAsset(const UInterchangeTextureFactor
 		return nullptr;
 	}
 
-	UTexture2D* Texture2D = Cast<UTexture2D>(Texture);
-	if (Texture2D)
-	{
-		if (BlockedPayloadData)
-		{
-			//
-			// Blocked Texture (also know as UDIM)
-			//
-			const UE::Interchange::FImportBlockedImage& BlockedImage = BlockedPayloadData.GetValue();
-
-			if (BlockedImage.IsBlockedData())
-			{
-				ETextureSourceFormat Format = BlockedImage.ImagesData[0].Format;
-				TextureCompressionSettings CompressionSettings = BlockedImage.ImagesData[0].CompressionSettings;
-				bool bSRGB = BlockedImage.ImagesData[0].SRGB;
-
-				TArray<const uint8*> SourceImageDatasPtr;
-				SourceImageDatasPtr.Reserve(BlockedImage.ImagesData.Num());
-				for (const  UE::Interchange::FImportImage& Image : BlockedImage.ImagesData)
-				{
-					SourceImageDatasPtr.Add(Image.RawData.GetData());
-				}
 
-				Texture2D->Source.InitBlocked(
-					&Format,
-					BlockedImage.BlocksData.GetData(),
-					/*NumSlices=*/ 1,
-					BlockedImage.BlocksData.Num(),
-					SourceImageDatasPtr.GetData()
-				);
-
-				Texture2D->CompressionSettings = CompressionSettings;
-				Texture2D->SRGB = bSRGB;
-				Texture2D->VirtualTextureStreaming = true;
-			}
-			else
-			{
-				//Import as a normal texture
-				PayloadData = MoveTemp(BlockedPayloadData.GetValue().ImagesData[0]);
-			}
-		}
-
-		if (PayloadData)
-		{
-			//
-			// Generic 2D Image
-			//
-			const UE::Interchange::FImportImage& Image = PayloadData.GetValue();
-			Texture2D->Source.Init(
-				Image.SizeX,
-				Image.SizeY,
-				/*NumSlices=*/ 1,
-				Image.NumMips,
-				Image.Format,
-				Image.RawData.GetData()
-			);
-			Texture2D->CompressionSettings = Image.CompressionSettings;
-			Texture2D->SRGB = Image.SRGB;
-
-			//If the MipGenSettings was set by the translator, we must apply it before the build
-			if (Image.MipGenSettings.IsSet())
-			{
-				// if the source has mips we keep the mips by default, unless the user changes that
-				Texture2D->MipGenSettings = Image.MipGenSettings.GetValue();
-			}
-
-		}
-
-
-		if(!Arguments.ReimportObject)
+	// Setup source data
+	if (UTexture2D* Texture2D = Cast<UTexture2D>(Texture))
+	{
+		if (!SetupTexture2DSourceData(Texture2D, TexturePayload))
 		{
-			/** Apply all TextureNode custom attributes to the texture asset */
-			TextureFactoryNode->ApplyAllCustomAttributeToAsset(Texture2D);
+			UE_LOG(LogInterchangeImport, Error, TEXT("UInterchangeTextureFactory: The Payload was invalid for a 2D Texture"), *Arguments.AssetName);
+			return nullptr;
 		}
-		else
+	}
+	else if (UTextureCube* TextureCube = Cast<UTextureCube>(Texture))
+	{
+		if (!SetupTextureCubeSourceData(TextureCube, TexturePayload))
 		{
-			UInterchangeAssetImportData* InterchangeAssetImportData = Cast<UInterchangeAssetImportData>(Texture2D->AssetImportData);
-			UInterchangeBaseNode* PreviousNode = nullptr;
-			if (InterchangeAssetImportData)
-			{
-				PreviousNode = InterchangeAssetImportData->NodeContainer->GetNode(InterchangeAssetImportData->NodeUniqueID);
-			}
-			UInterchangeTextureFactoryNode* CurrentNode = NewObject<UInterchangeTextureFactoryNode>();
-			UInterchangeBaseNode::CopyStorage(TextureFactoryNode, CurrentNode);
-			CurrentNode->FillAllCustomAttributeFromAsset(Texture2D);
-			//Apply reimport strategy
-			UE::Interchange::FFactoryCommon::ApplyReimportStrategyToAsset(Arguments.ReimportStrategyFlags, Texture2D, PreviousNode, CurrentNode, TextureFactoryNode);
+			UE_LOG(LogInterchangeImport, Error, TEXT("UInterchangeTextureFactory: The Payload was invalid for a TextureCube"), *Arguments.AssetName);
+			return nullptr;
 		}
-		
-		//Getting the file Hash will cache it into the source data
-		Arguments.SourceData->GetFileContentHash();
-
-		//The interchange completion task (call in the GameThread after the factories pass), will call PostEditChange which will trig another asynchronous system that will build all texture in parallel
 	}
 	else
 	{
 		//The texture is not a UTexture2D
 		Texture->RemoveFromRoot();
 		Texture->MarkPendingKill();
+		return nullptr;
+	}
+
+	UInterchangeBaseNode* TextureFactoryNode = Arguments.AssetNode;
+	if(!Arguments.ReimportObject)
+	{
+		/** Apply all TextureNode custom attributes to the texture asset */
+		TextureFactoryNode->ApplyAllCustomAttributeToAsset(Texture);
 	}
-	return Texture2D;
+	else
+	{
+		UInterchangeAssetImportData* InterchangeAssetImportData = Cast<UInterchangeAssetImportData>(Texture->AssetImportData);
+		UInterchangeBaseNode* PreviousNode = nullptr;
+		if (InterchangeAssetImportData)
+		{
+			PreviousNode = InterchangeAssetImportData->NodeContainer->GetNode(InterchangeAssetImportData->NodeUniqueID);
+		}
+
+		UInterchangeBaseNode* CurrentNode = NewObject<UInterchangeBaseNode>(GetTransientPackage(), SupportedFactoryNodeClass);
+		UInterchangeBaseNode::CopyStorage(TextureFactoryNode, CurrentNode);
+		CurrentNode->FillAllCustomAttributeFromAsset(Texture);
+		//Apply reimport strategy
+		UE::Interchange::FFactoryCommon::ApplyReimportStrategyToAsset(Arguments.ReimportStrategyFlags, Texture, PreviousNode, CurrentNode, TextureFactoryNode);
+	}
+		
+	//Getting the file Hash will cache it into the source data
+	Arguments.SourceData->GetFileContentHash();
+
+	//The interchange completion task (call in the GameThread after the factories pass), will call PostEditChange which will trig another asynchronous system that will build all texture in parallel
+
+	return Texture;
 #endif
 }
 
diff --git a/Engine/Plugins/Runtime/Interchange/Source/Import/Private/Texture/TextureTranslatorUtilities.cpp b/Engine/Plugins/Runtime/Interchange/Source/Import/Private/Texture/TextureTranslatorUtilities.cpp
--- a/Engine/Plugins/Runtime/Interchange/Source/Import/Private/Texture/TextureTranslatorUtilities.cpp
+++ b/Engine/Plugins/Runtime/Interchange/Source/Import/Private/Texture/TextureTranslatorUtilities.cpp
@@ -3,6 +3,7 @@
 
 #include "CoreMinimal.h"
 #include "InterchangeSourceData.h"
+#include "InterchangeTextureCubeNode.h"
 #include "InterchangeTextureNode.h"
 #include "Misc/Paths.h"
 #include "Nodes/InterchangeBaseNodeContainer.h"
@@ -33,3 +34,29 @@ bool UE::Interchange::FTextureTranslatorUtilities::Generic2DTextureTranslate(con
 
 	return true;
 }
+
+bool UE::Interchange::FTextureTranslatorUtilities::GenericCubeTextureTranslate(const UInterchangeSourceData* SourceData, UInterchangeBaseNodeContainer& BaseNodeContainer)
+{
+	FString Filename = SourceData->GetFilename();
+	FPaths::NormalizeFilename(Filename);
+	if (!FPaths::FileExists(Filename))
+	{
+		return false;
+	}
+
+	FString DisplayLabel = FPaths::GetBaseFilename(Filename);
+	FString NodeUID(Filename);
+	UInterchangeTextureCubeNode* TextureCubeNode = NewObject<UInterchangeTextureCubeNode>(&BaseNodeContainer, NAME_None);
+	if (!ensure(TextureCubeNode))
+	{
+		return false;
+	}
+
+	//Creating a Cube Texture
+	TextureCubeNode->InitializeNode(NodeUID, DisplayLabel, EInterchangeNodeContainerType::NodeContainerType_TranslatedAsset);
+	TextureCubeNode->SetPayLoadKey(Filename);
+
+	BaseNodeContainer.AddNode(TextureCubeNode);
+
+	return true;
+}
diff --git a/Engine/Plugins/Runtime/Interchange/Source/Import/Private/Texture/TextureTranslatorUtilities.h b/Engine/Plugins/Runtime/Interchange/Source/Import/Private/Texture/TextureTranslatorUtilities.h
--- a/Engine/Plugins/Runtime/Interchange/Source/Import/Private/Texture/TextureTranslatorUtilities.h
+++ b/Engine/Plugins/Runtime/Interchange/Source/Import/Private/Texture/TextureTranslatorUtilities.h
@@ -15,6 +15,8 @@ namespace UE
 		{
 		public:
 			static bool Generic2DTextureTranslate(const UInterchangeSourceData* SourceData, UInterchangeBaseNodeContainer& BaseNodeContainer);
+
+			static bool GenericCubeTextureTranslate(const UInterchangeSourceData* SourceData, UInterchangeBaseNodeContainer& BaseNodeContainer);
 		};
 	}//ns Interchange
 }//ns UE
diff --git a/Engine/Plugins/Runtime/Interchange/Source/Import/Public/Texture/InterchangeDDSTranslator.h b/Engine/Plugins/Runtime/Interchange/Source/Import/Public/Texture/InterchangeDDSTranslator.h
--- a/Engine/Plugins/Runtime/Interchange/Source/Import/Public/Texture/InterchangeDDSTranslator.h
+++ b/Engine/Plugins/Runtime/Interchange/Source/Import/Public/Texture/InterchangeDDSTranslator.h
@@ -3,17 +3,19 @@
 #pragma once
 
 #include "CoreMinimal.h"
-#include "UObject/Object.h"
-#include "UObject/ObjectMacros.h"
 #include "InterchangeTranslatorBase.h"
 #include "Nodes/InterchangeBaseNodeContainer.h"
+#include "Texture/InterchangeSlicedTexturePayloadData.h"
+#include "Texture/InterchangeSlicedTexturePayloadInterface.h"
 #include "Texture/InterchangeTexturePayloadData.h"
 #include "Texture/InterchangeTexturePayloadInterface.h"
+#include "UObject/Object.h"
+#include "UObject/ObjectMacros.h"
 
 #include "InterchangeDDSTranslator.generated.h"
 
 UCLASS(BlueprintType)
-class INTERCHANGEIMPORT_API UInterchangeDDSTranslator : public UInterchangeTranslatorBase, public IInterchangeTexturePayloadInterface
+class INTERCHANGEIMPORT_API UInterchangeDDSTranslator : public UInterchangeTranslatorBase, public IInterchangeTexturePayloadInterface, public IInterchangeSlicedTexturePayloadInterface
 {
 	GENERATED_BODY()
 public:
@@ -45,6 +47,21 @@ public:
 	virtual TOptional<UE::Interchange::FImportImage> GetTexturePayloadData(const UInterchangeSourceData* SourceData, const FString& PayLoadKey) const override;
 
 	/* IInterchangeTexturePayloadInterface End */
+
+
+	/* IInterchangeSlicedTexturePayloadInterface Begin */
+
+	/**
+	 * Once the translation is done, the import process need a way to retrieve payload data.
+	 * This payload will be use by the factories to create the asset.
+	 *
+	 * @param SourceData - The source data containing the data to translate
+	 * @param PayloadKey - The key to retrieve the a particular payload contain into the specified source data.
+	 * @return a PayloadData containing the import image data. The TOptional will not be set if there is an error.
+	 */
+	virtual TOptional<UE::Interchange::FImportSlicedImage> GetSlicedTexturePayloadData(const UInterchangeSourceData* SourceData, const FString& PayloadKey) const override;
+
+	/* IInterchangeSlicedTexturePayloadInterface End */
 };
 
 
diff --git a/Engine/Plugins/Runtime/Interchange/Source/Import/Public/Texture/InterchangeSlicedTexturePayloadData.h b/Engine/Plugins/Runtime/Interchange/Source/Import/Public/Texture/InterchangeSlicedTexturePayloadData.h
new file mode 100644
--- /dev/null
+++ b/Engine/Plugins/Runtime/Interchange/Source/Import/Public/Texture/InterchangeSlicedTexturePayloadData.h
@@ -0,0 +1,33 @@
+// Copyright Epic Games, Inc. All Rights Reserved.
+
+#pragma once
+
+#include "CoreMinimal.h"
+#include "Engine/Texture.h"
+
+namespace UE
+{
+	namespace Interchange
+	{
+		struct INTERCHANGEIMPORT_API FImportSlicedImage
+		{
+			TArray64<uint8> RawData;
+			ETextureSourceFormat Format = TSF_Invalid;
+			TextureCompressionSettings CompressionSettings = TC_Default;
+			int32 NumMips = 0;
+			int32 SizeX = 0;
+			int32 SizeY = 0;
+			int32 NumSlice = 0;
+			bool bSRGB = true;
+			TOptional<TextureMipGenSettings> MipGenSettings;
+
+			void Init(int32 InSizeX, int32 InSizeY, int32 InNumSlice, int32 InNumMips, ETextureSourceFormat InFormat , bool InSRGB);
+
+			/** Return the size of a mip for one slice */
+			int64 GetMipSize(int32 InMipIndex) const;
+
+			const uint8* GetMipData(int32 InMipIndex, int32 InSliceIndex = INDEX_NONE) const;
+			uint8* GetMipData(int32 InMipIndex, int32 InSliceIndex = INDEX_NONE);
+		};
+	}
+}
diff --git a/Engine/Plugins/Runtime/Interchange/Source/Import/Public/Texture/InterchangeSlicedTexturePayloadInterface.h b/Engine/Plugins/Runtime/Interchange/Source/Import/Public/Texture/InterchangeSlicedTexturePayloadInterface.h
new file mode 100644
--- /dev/null
+++ b/Engine/Plugins/Runtime/Interchange/Source/Import/Public/Texture/InterchangeSlicedTexturePayloadInterface.h
@@ -0,0 +1,35 @@
+// Copyright Epic Games, Inc. All Rights Reserved.
+#pragma once
+
+#include "CoreMinimal.h"
+#include "InterchangeSourceData.h"
+#include "Texture/InterchangeSlicedTexturePayloadData.h"
+#include "UObject/Object.h"
+#include "UObject/ObjectMacros.h"
+
+#include "InterchangeSlicedTexturePayloadInterface.generated.h"
+
+UINTERFACE()
+class INTERCHANGEIMPORT_API UInterchangeSlicedTexturePayloadInterface : public UInterface
+{
+	GENERATED_BODY()
+};
+
+/**
+ * Texture payload interface. Derive from it if your payload can import texture
+ */
+class INTERCHANGEIMPORT_API IInterchangeSlicedTexturePayloadInterface
+{
+	GENERATED_BODY()
+public:
+
+	/**
+	 * Once the translation is done, the import process need a way to retrieve payload data.
+	 * This payload will be use by the factories to create the asset.
+	 *
+	 * @param SourceData - The source data containing the data to translate
+	 * @param PayloadKey - The key to retrieve the a particular payload contain into the specified source data.
+	 * @return a PayloadData containing the import image data. The TOptional will not be set if there is an error.
+	 */
+	virtual TOptional<UE::Interchange::FImportSlicedImage> GetSlicedTexturePayloadData(const UInterchangeSourceData* SourceData, const FString& PayloadKey) const = 0;
+};
diff --git a/Engine/Plugins/Runtime/Interchange/Source/Nodes/Public/InterchangeTextureCubeFactoryNode.h b/Engine/Plugins/Runtime/Interchange/Source/Nodes/Public/InterchangeTextureCubeFactoryNode.h
new file mode 100644
--- /dev/null
+++ b/Engine/Plugins/Runtime/Interchange/Source/Nodes/Public/InterchangeTextureCubeFactoryNode.h
@@ -0,0 +1,33 @@
+// Copyright Epic Games, Inc. All Rights Reserved.
+
+#pragma once
+
+#include "CoreMinimal.h"
+#include "InterchangeTextureFactoryNode.h"
+
+#if WITH_ENGINE
+#include "Engine/TextureCube.h"
+#endif
+
+
+#include "InterchangeTextureCubeFactoryNode.generated.h"
+
+UCLASS(BlueprintType)
+class INTERCHANGENODES_API UInterchangeTextureCubeFactoryNode : public UInterchangeTextureFactoryNode
+{
+	GENERATED_BODY()
+
+private:
+
+#if WITH_ENGINE
+	virtual void FillAssetClassFromClassName(const FString& ClassName) override
+	{
+		if (ClassName.Equals(UTextureCube::StaticClass()->GetName()))
+		{
+			AssetClass = UTextureCube::StaticClass();
+			bIsTextureNodeClassInitialized = true;
+		}
+	}
+#endif
+
+};
\ No newline at end of file
diff --git a/Engine/Plugins/Runtime/Interchange/Source/Nodes/Public/InterchangeTextureCubeNode.h b/Engine/Plugins/Runtime/Interchange/Source/Nodes/Public/InterchangeTextureCubeNode.h
new file mode 100644
--- /dev/null
+++ b/Engine/Plugins/Runtime/Interchange/Source/Nodes/Public/InterchangeTextureCubeNode.h
@@ -0,0 +1,89 @@
+// Copyright Epic Games, Inc. All Rights Reserved.
+
+#pragma once
+
+#include "CoreMinimal.h"
+#include "Nodes/InterchangeBaseNode.h"
+
+#include "InterchangeTextureCubeNode.generated.h"
+
+namespace UE
+{
+	namespace Interchange
+	{
+
+		struct FTextureCubeNodeStaticData : public FBaseNodeStaticData
+		{
+			static const FAttributeKey& PayloadSourceFileKey()
+			{
+				static FAttributeKey AttributeKey(TEXT("__PayloadSourceFile__"));
+				return AttributeKey;
+			}
+		};
+	}
+}
+
+
+UCLASS(BlueprintType)
+class INTERCHANGENODES_API UInterchangeTextureCubeNode : public UInterchangeBaseNode
+{
+	GENERATED_BODY()
+
+public:
+
+	/**
+	 * Return the node type name of the class, we use this when reporting error
+	 */
+	virtual FString GetTypeName() const override
+	{
+		const FString TypeName = TEXT("TextureCubeNode");
+		return TypeName;
+	}
+
+	virtual FString GetKeyDisplayName(const UE::Interchange::FAttributeKey& NodeAttributeKey) const override
+	{
+		FString KeyDisplayName = NodeAttributeKey.ToString();
+		if (NodeAttributeKey == UE::Interchange::FTextureCubeNodeStaticData::PayloadSourceFileKey())
+		{
+			return KeyDisplayName = TEXT("Payload Source Key");
+		}
+		return Super::GetKeyDisplayName(NodeAttributeKey);
+	}
+
+	virtual FGuid GetHash() const override
+	{
+		return Attributes->GetStorageHash();
+	}
+
+	/** Texture node Interface Begin */
+	virtual const TOptional<FString> GetPayLoadKey() const
+	{
+		if (!Attributes->ContainAttribute(UE::Interchange::FTextureCubeNodeStaticData::PayloadSourceFileKey()))
+		{
+			return TOptional<FString>();
+		}
+		UE::Interchange::FAttributeStorage::TAttributeHandle<FString> AttributeHandle = Attributes->GetAttributeHandle<FString>(UE::Interchange::FTextureCubeNodeStaticData::PayloadSourceFileKey());
+		if (!AttributeHandle.IsValid())
+		{
+			return TOptional<FString>();
+		}
+		FString PayloadKey;
+		UE::Interchange::EAttributeStorageResult Result = AttributeHandle.Get(PayloadKey);
+		if (!IsAttributeStorageResultSuccess(Result))
+		{
+			LogAttributeStorageErrors(Result, TEXT("UInterchangeTextureNode.GetPayLoadKey"), UE::Interchange::FTextureCubeNodeStaticData::PayloadSourceFileKey());
+			return TOptional<FString>();
+		}
+		return TOptional<FString>(PayloadKey);
+	}
+
+	UFUNCTION(BlueprintCallable, Category = "Interchange | Node | Texture")
+	virtual void SetPayLoadKey(const FString& PayloadKey)
+	{
+		UE::Interchange::EAttributeStorageResult Result = Attributes->RegisterAttribute(UE::Interchange::FTextureCubeNodeStaticData::PayloadSourceFileKey(), PayloadKey);
+		if (!IsAttributeStorageResultSuccess(Result))
+		{
+			LogAttributeStorageErrors(Result, TEXT("UInterchangeTextureNode.SetPayLoadKey"), UE::Interchange::FTextureCubeNodeStaticData::PayloadSourceFileKey());
+		}
+	}
+};
diff --git a/Engine/Plugins/Runtime/Interchange/Source/Nodes/Public/InterchangeTextureFactoryNode.h b/Engine/Plugins/Runtime/Interchange/Source/Nodes/Public/InterchangeTextureFactoryNode.h
--- a/Engine/Plugins/Runtime/Interchange/Source/Nodes/Public/InterchangeTextureFactoryNode.h
+++ b/Engine/Plugins/Runtime/Interchange/Source/Nodes/Public/InterchangeTextureFactoryNode.h
@@ -720,6 +720,13 @@ private:
 		FString OperationName = GetTypeName() + TEXT(".GetAssetClassName");
 		FString ClassName;
 		InterchangePrivateNodeBase::GetCustomAttribute<FString>(*Attributes, ClassNameAttributeKey, OperationName, ClassName);
+		FillAssetClassFromClassName(ClassName);
+#endif
+	}
+
+#if WITH_ENGINE
+	virtual void FillAssetClassFromClassName(const FString& ClassName)
+	{
 		if (ClassName.Equals(UTexture2D::StaticClass()->GetName()))
 		{
 			AssetClass = UTexture2D::StaticClass();
@@ -730,8 +737,8 @@ private:
 			AssetClass = UTexture::StaticClass();
 			bIsTextureNodeClassInitialized = true;
 		}
-#endif
 	}
+#endif
 
 	bool SetTextureNodeClassFromClassAttribute()
 	{
@@ -896,7 +903,5 @@ protected:
 	TSubclassOf<UTexture> AssetClass = nullptr;
 #endif
 
-	UE::Interchange::TMapAttributeHelper<int32, FString> SourceBlocks;
-
 	bool bIsTextureNodeClassInitialized = false;
 };
diff --git a/Engine/Plugins/Runtime/Interchange/Source/Pipelines/Private/InterchangeGenericAssetsPipeline.cpp b/Engine/Plugins/Runtime/Interchange/Source/Pipelines/Private/InterchangeGenericAssetsPipeline.cpp
--- a/Engine/Plugins/Runtime/Interchange/Source/Pipelines/Private/InterchangeGenericAssetsPipeline.cpp
+++ b/Engine/Plugins/Runtime/Interchange/Source/Pipelines/Private/InterchangeGenericAssetsPipeline.cpp
@@ -3,6 +3,7 @@
 
 #include "Animation/Skeleton.h"
 #include "CoreMinimal.h"
+#include "Engine/TextureCube.h"
 #include "InterchangeMaterialFactoryNode.h"
 #include "InterchangeMaterialNode.h"
 #include "InterchangeMeshNode.h"
@@ -12,9 +13,11 @@
 #include "InterchangeSkeletalMeshFactoryNode.h"
 #include "InterchangeSkeletalMeshLodDataNode.h"
 #include "InterchangeSkeletonFactoryNode.h"
+#include "InterchangeSourceData.h"
+#include "InterchangeTextureCubeFactoryNode.h"
+#include "InterchangeTextureCubeNode.h"
 #include "InterchangeTextureFactoryNode.h"
 #include "InterchangeTextureNode.h"
-#include "InterchangeSourceData.h"
 #include "Misc/Paths.h"
 #include "Nodes/InterchangeBaseNode.h"
 #include "Nodes/InterchangeBaseNodeContainer.h"
@@ -58,6 +61,7 @@ bool UInterchangeGenericAssetsPipeline::ExecutePreImportPipeline(UInterchangeBas
 				}
 				else if (UInterchangeMeshNode* MeshNode = Cast<UInterchangeMeshNode>(Node))
 				{
+
 					MeshNodes.Add(MeshNode);
 					if (MeshNode->IsSkinnedMesh())
 					{
@@ -69,6 +73,10 @@ bool UInterchangeGenericAssetsPipeline::ExecutePreImportPipeline(UInterchangeBas
 					}
 
 				}
+				else if (UInterchangeTextureCubeNode* TextuteCubeNode = Cast<UInterchangeTextureCubeNode>(Node))
+				{
+					TextureCubeNodes.Add(TextuteCubeNode);
+				}
 			}
 			break;
 			case EInterchangeNodeContainerType::NodeContainerType_TranslatedScene:
@@ -91,6 +99,12 @@ bool UInterchangeGenericAssetsPipeline::ExecutePreImportPipeline(UInterchangeBas
 		}
 	}
 
+	for (const UInterchangeTextureCubeNode* TextureCubeNode : TextureCubeNodes)
+	{
+		CreateTextureCubeFactoryNode(TextureCubeNode);
+	}
+
+
 	if (bImportMaterials)
 	{
 		//import materials
@@ -321,6 +335,36 @@ UInterchangeTextureFactoryNode* UInterchangeGenericAssetsPipeline::CreateTexture
 	return TextureFactoryNode;
 }
 
+UInterchangeTextureCubeFactoryNode* UInterchangeGenericAssetsPipeline::CreateTextureCubeFactoryNode(const UInterchangeTextureCubeNode* TextureCubeNode)
+{
+	FString DisplayLabel = TextureCubeNode->GetDisplayLabel();
+	FString NodeUid = UInterchangeTextureFactoryNode::GetTextureFactoryNodeUidFromTextureNodeUid(TextureCubeNode->GetUniqueID());
+	UInterchangeTextureCubeFactoryNode* TextureCubeFactoryNode = nullptr;
+	if (BaseNodeContainer->IsNodeUidValid(NodeUid))
+	{
+		TextureCubeFactoryNode = Cast<UInterchangeTextureCubeFactoryNode>(BaseNodeContainer->GetNode(NodeUid));
+		if (!ensure(TextureCubeFactoryNode))
+		{
+			//Log an error
+		}
+	}
+	else
+	{
+		TextureCubeFactoryNode = NewObject<UInterchangeTextureCubeFactoryNode>(BaseNodeContainer, NAME_None);
+		if (!ensure(TextureCubeFactoryNode))
+		{
+			return nullptr;
+		}
+
+		//Creating a UTextureCube
+		TextureCubeFactoryNode->InitializeTextureNode(NodeUid, DisplayLabel, UTextureCube::StaticClass()->GetName(), TextureCubeNode->GetDisplayLabel());
+		TextureCubeFactoryNode->SetCustomTranslatedTextureNodeUid(TextureCubeNode->GetUniqueID());
+		BaseNodeContainer->AddNode(TextureCubeFactoryNode);
+		TextureCubeFactoryNodes.Add(TextureCubeFactoryNode);
+	}
+	return TextureCubeFactoryNode;
+}
+
 UInterchangeMaterialFactoryNode* UInterchangeGenericAssetsPipeline::CreateMaterialFactoryNode(const UInterchangeMaterialNode* MaterialNode)
 {
 	FString DisplayLabel = MaterialNode->GetDisplayLabel();
diff --git a/Engine/Plugins/Runtime/Interchange/Source/Pipelines/Public/InterchangeGenericAssetsPipeline.h b/Engine/Plugins/Runtime/Interchange/Source/Pipelines/Public/InterchangeGenericAssetsPipeline.h
--- a/Engine/Plugins/Runtime/Interchange/Source/Pipelines/Public/InterchangeGenericAssetsPipeline.h
+++ b/Engine/Plugins/Runtime/Interchange/Source/Pipelines/Public/InterchangeGenericAssetsPipeline.h
@@ -11,16 +11,18 @@
 
 #include "InterchangeGenericAssetsPipeline.generated.h"
 
-class UInterchangeTextureNode;
-class UInterchangeTextureFactoryNode;
-class UInterchangeMaterialNode;
 class UInterchangeMaterialFactoryNode;
+class UInterchangeMaterialNode;
 class UInterchangeMeshNode;
 class UInterchangePipelineMeshesUtilities;
 class UInterchangeSceneNode;
 class UInterchangeSkeletalMeshFactoryNode;
 class UInterchangeSkeletalMeshLodDataNode;
 class UInterchangeSkeletonFactoryNode;
+class UInterchangeTextureCubeFactoryNode;
+class UInterchangeTextureCubeNode;
+class UInterchangeTextureFactoryNode;
+class UInterchangeTextureNode;
 
 #define COMMON_CATEGORY "Common"
 #define COMMON_MESHES_CATEGORY "Common Meshes"
@@ -159,14 +161,24 @@ private:
 	/** Texture factory assets nodes */
 	TArray<UInterchangeTextureFactoryNode*> TextureFactoryNodes;
 
+
+	/** Texture cube translated assets nodes*/
+	TArray<UInterchangeTextureCubeNode*> TextureCubeNodes;
+
+	/** Texture cube factory translated assets nodes */
+	TArray<UInterchangeTextureCubeFactoryNode*> TextureCubeFactoryNodes;
+
+
 	UInterchangeTextureFactoryNode* CreateTexture2DFactoryNode(const UInterchangeTextureNode* TextureNode);
 
+	UInterchangeTextureCubeFactoryNode* CreateTextureCubeFactoryNode(const UInterchangeTextureCubeNode* TextureCubeNode);
 	
 	/** Material translated assets nodes */
 	TArray<UInterchangeMaterialNode*> MaterialNodes;
 	
 	/** Material factory assets nodes */
 	TArray<UInterchangeMaterialFactoryNode*> MaterialFactoryNodes;
+
 	
 	UInterchangeMaterialFactoryNode* CreateMaterialFactoryNode(const UInterchangeMaterialNode* MaterialNode);
 
diff --git a/Engine/Source/Editor/UnrealEd/Private/Factories/EditorFactories.cpp b/Engine/Source/Editor/UnrealEd/Private/Factories/EditorFactories.cpp
--- a/Engine/Source/Editor/UnrealEd/Private/Factories/EditorFactories.cpp
+++ b/Engine/Source/Editor/UnrealEd/Private/Factories/EditorFactories.cpp
@@ -3800,8 +3800,10 @@ UTexture* UTextureFactory::ImportTexture(UClass* Class, UObject* InParent, FName
 				TextureCube->Source.UnlockMip(MipIndex);
 			}
 
-			// for now we don't support mip map generation on cubemaps
-			TextureCube->MipGenSettings = TMGS_LeaveExistingMips;
+			if (NumMips > 1)
+			{
+				TextureCube->MipGenSettings = TMGS_LeaveExistingMips;
+			}
 		}
 
 		return TextureCube;
@@ -4312,35 +4314,6 @@ UObject* UTextureFactory::FactoryCreateBinary
 
 	UTexture2D* Texture2D = Cast<UTexture2D>(Texture);
 
-	if (Texture2D)
-	{
-		// If the texture is larger than a certain threshold make it VT.
-		// Note that previously for re-imports we still checked size and potentially changed the VT status.
-		// But that was unintuitive for many users so now for re-imports we will end up ignoring this and respecting the existing setting below.
-		static const auto CVarVirtualTexturesEnabled = IConsoleManager::Get().FindTConsoleVariableDataInt(TEXT("r.VirtualTextures")); 
-		check(CVarVirtualTexturesEnabled);
-		static const auto CVarVirtualTexturesAutoImportEnabled = IConsoleManager::Get().FindTConsoleVariableDataInt(TEXT("r.VT.EnableAutoImport"));
-		check(CVarVirtualTexturesAutoImportEnabled);
-
-		if (CVarVirtualTexturesEnabled->GetValueOnAnyThread() && CVarVirtualTexturesAutoImportEnabled->GetValueOnAnyThread())
-		{
-			const int32 VirtualTextureAutoEnableThreshold = GetDefault<UTextureImportSettings>()->AutoVTSize;
-			const int32 VirtualTextureAutoEnableThresholdPixels = VirtualTextureAutoEnableThreshold * VirtualTextureAutoEnableThreshold;
-
-			// We do this in pixels so a 8192 x 128 texture won't get VT enabled 
-			// We use the Source size instead of simple Texture2D->GetSizeX() as this uses the size of the platform data
-			// however for a new texture platform data may not be generated yet, and for an reimport of a texture this is the size of the
-			// old texture. 
-			// Using source size gives one small caveat. It looks at the size before mipmap power of two padding adjustment.
-			// Textures with more than 1 block (UDIM textures) must be imported as VT
-			if (Texture->Source.GetNumBlocks() > 1 ||
-				Texture2D->Source.GetSizeX() * Texture2D->Source.GetSizeY() >= VirtualTextureAutoEnableThresholdPixels)
-			{
-				Texture2D->VirtualTextureStreaming = true;
-			}
-		}
-	}
-
 	// Restore user set options
 	if (ExistingTexture && bUsingExistingSettings)
 	{
@@ -4386,6 +4359,29 @@ UObject* UTextureFactory::FactoryCreateBinary
 	{
 		// The texture has been imported and has no editor specific changes applied so we clear the painted flag.
 		Texture2D->bHasBeenPaintedInEditor = false;
+
+		// If the texture is larger than a certain threshold make it VT. This is explicitly done after the
+		// application of the existing settings above, so if a texture gets reimported at a larger size it will
+		// still be properly flagged as a VT (note: What about reimporting at a lower resolution?)
+		static const auto CVarVirtualTexturesEnabled = IConsoleManager::Get().FindTConsoleVariableDataInt(TEXT("r.VirtualTextures")); check(CVarVirtualTexturesEnabled);
+		
+		if (CVarVirtualTexturesEnabled->GetValueOnAnyThread())
+		{
+			int virtualTextureAutoEnableThreshold = GetDefault<UTextureImportSettings>()->AutoVTSize;
+			int virtualTextureAutoEnableThresholdPixels = virtualTextureAutoEnableThreshold * virtualTextureAutoEnableThreshold;
+
+			// We do this in pixels so a 8192 x 128 texture won't get VT enabled 
+			// We use the Source size instead of simple Texture2D->GetSizeX() as this uses the size of the platform data
+			// however for a new texture platform data may not be generated yet, and for an reimport of a texture this is the size of the
+			// old texture. 
+			// Using source size gives one small caveat. It looks at the size before mipmap power of two padding adjustment.
+			// Textures with more than 1 block (UDIM textures) must be imported as VT
+			if (Texture->Source.GetNumBlocks() > 1 ||
+				Texture2D->Source.GetSizeX()*Texture2D->Source.GetSizeY() >= virtualTextureAutoEnableThresholdPixels)
+			{
+				Texture2D->VirtualTextureStreaming = true;
+			}
+		}
 	}
 
 	// Automatically detect if the texture is a normal map and configure its properties accordingly
diff --git a/Engine/Source/Runtime/Engine/Private/DDSLoader.cpp b/Engine/Source/Runtime/Engine/Private/DDSLoader.cpp
--- a/Engine/Source/Runtime/Engine/Private/DDSLoader.cpp
+++ b/Engine/Source/Runtime/Engine/Private/DDSLoader.cpp
@@ -5,7 +5,7 @@
 #include "Engine/Texture2D.h"
 #include "RenderUtils.h"
 
-FDDSLoadHelper::FDDSLoadHelper(const uint8* Buffer, uint32 Length) 
+FDDSLoadHelper::FDDSLoadHelper(const uint8* Buffer, uint64 Length) 
 	: DDSHeader(0), DDS10Header(0)
 {
 	check(Buffer);
@@ -151,17 +151,18 @@ ETextureSourceFormat FDDSLoadHelper::ComputeSourceFormat() const
 
 bool FDDSLoadHelper::IsValidCubemapTexture() const
 {
-	if(DDSHeader != nullptr && (DDSHeader->dwCaps2 & DDSC_CubeMap) != 0 && (DDSHeader->dwCaps2 & DDSC_CubeMap_AllFaces) != 0)
-	{
-		return true;
-	}
+	if (DDSHeader != nullptr  && DDSHeader->dwWidth == DDSHeader->dwHeight)
+	{ 
+		if((DDSHeader->dwCaps2 & DDSC_CubeMap) != 0 && (DDSHeader->dwCaps2 & DDSC_CubeMap_AllFaces) != 0)
+		{
+			return true;
+		}
 
-	if (DDS10Header != nullptr && DDS10Header->resourceType == 3 && (DDS10Header->miscFlag & 4))
-	{
-		return true;
+		if (DDS10Header != nullptr && DDS10Header->resourceType == 3 && (DDS10Header->miscFlag & 4))
+		{
+			return true;
+		}
 	}
-
-
 	return false;
 }
 
@@ -217,7 +218,7 @@ uint32 FDDSLoadHelper::ComputeMipMapCount() const
 
 const uint8* FDDSLoadHelper::GetDDSDataPointer(ECubeFace Face) const
 {
-	uint32 SliceSize = CalcTextureSize(DDSHeader->dwWidth, DDSHeader->dwHeight, ComputePixelFormat(), ComputeMipMapCount());
+	uint32 SliceSize = CalcTextureSize(GetSizeX(), GetSizeY(), ComputePixelFormat(), ComputeMipMapCount());
 
 	const uint8* Ptr = (const uint8*)DDSHeader + sizeof(FDDSFileHeader);
 
@@ -246,3 +247,14 @@ const uint8* FDDSLoadHelper::GetDDSDataPointer(const UTexture2D& Texture) const
 
 	return GetDDSDataPointer();
 }
+
+int64 FDDSLoadHelper::GetDDSHeaderMaximalSize()
+{
+	return sizeof(FDDSFileHeader) + sizeof(FDDS10FileHeader);
+}
+
+int64 FDDSLoadHelper::GetDDSHeaderMinimalSize()
+{
+	return sizeof(FDDSFileHeader);
+}
+
diff --git a/Engine/Source/Runtime/Engine/Public/DDSLoader.h b/Engine/Source/Runtime/Engine/Public/DDSLoader.h
--- a/Engine/Source/Runtime/Engine/Public/DDSLoader.h
+++ b/Engine/Source/Runtime/Engine/Public/DDSLoader.h
@@ -90,7 +90,7 @@ class FDDSLoadHelper
 {
 public:
 	/** @param Buffer must not be 0 */
-	ENGINE_API FDDSLoadHelper(const uint8* Buffer, uint32 Length);
+	ENGINE_API FDDSLoadHelper(const uint8* Buffer, uint64 Length);
 
 	bool IsValid() const;
 		
@@ -120,7 +120,13 @@ public:
 
 	/** Extracts the cube map size from the texture name */
 	ENGINE_API const uint8* GetDDSDataPointer(const UTexture2D& Texture) const;
-	
+
+	/** Get the largest size possible for a DDS header in Bytes */
+	ENGINE_API static int64 GetDDSHeaderMaximalSize();
+
+	/** Get the minimal size possible for a DDS header in Bytes */
+	ENGINE_API static int64 GetDDSHeaderMinimalSize();
+
 // for now we allow direct access to the header but that could be changed
 
 	/** !=0 if valid */
