Compilation time optimization
  - Use Unity and PCH and enable back the optimizations
  - Probably submitted by mistake
  - 656s -> 630s for full QAGame rebuild on 3970X without XGE

diff --git a/Engine/Plugins/Runtime/Interchange/Source/Import/InterchangeImport.Build.cs b/Engine/Plugins/Runtime/Interchange/Source/Import/InterchangeImport.Build.cs
--- a/Engine/Plugins/Runtime/Interchange/Source/Import/InterchangeImport.Build.cs
+++ b/Engine/Plugins/Runtime/Interchange/Source/Import/InterchangeImport.Build.cs
@@ -31,10 +31,6 @@ namespace UnrealBuildTool.Rules
 					"TextureUtilitiesCommon",
 				}
 			);
-
-			OptimizeCode = CodeOptimization.Never;
-			bUseUnity = false;
-			PCHUsage = PCHUsageMode.NoPCHs;
 		}
 	}
 }
