Fix static analysis warnings

diff --git a/Engine/Plugins/Runtime/Interchange/Source/Import/Private/Texture/InterchangePCXTranslator.cpp b/Engine/Plugins/Runtime/Interchange/Source/Import/Private/Texture/InterchangePCXTranslator.cpp
--- a/Engine/Plugins/Runtime/Interchange/Source/Import/Private/Texture/InterchangePCXTranslator.cpp
+++ b/Engine/Plugins/Runtime/Interchange/Source/Import/Private/Texture/InterchangePCXTranslator.cpp
@@ -117,7 +117,7 @@ TOptional<UE::Interchange::FImportImage> UInterchangePCXTranslator::GetTexturePa
 	// Validate it.
 	const int32 Length = BufferEnd - Buffer;
 
-	IImageWrapperModule& ImageWrapperModule = FModuleManager::LoadModuleChecked<IImageWrapperModule>(FName("ImageWrapper"));
+	FModuleManager::LoadModuleChecked<IImageWrapperModule>(FName("ImageWrapper"));
 
 	UE::Interchange::FImportImage PayloadData;
 
diff --git a/Engine/Source/Runtime/VulkanRHI/Private/VulkanTexture.cpp b/Engine/Source/Runtime/VulkanRHI/Private/VulkanTexture.cpp
--- a/Engine/Source/Runtime/VulkanRHI/Private/VulkanTexture.cpp
+++ b/Engine/Source/Runtime/VulkanRHI/Private/VulkanTexture.cpp
@@ -875,7 +875,6 @@ void FVulkanSurface::OnFullDefrag(FVulkanDevice& InDevice, FVulkanCommandListCon
 
 		{
 			check(Image != VK_NULL_HANDLE);
-			uint64 Size = GetMemorySize();
 			Device->NotifyDeletedImage(Image, bRenderTarget);
 			Device->GetDeferredDeletionQueue().EnqueueResource(VulkanRHI::FDeferredDeletionQueue2::EType::Image, Image);
 			FGenericPlatformMisc::LowLevelOutputDebugStringf(TEXT("** MOVE IMAGE %p -> %p\n"), Image, MovedImage);
