Fix dispatcher Log spam
#rnx

diff --git a/Engine/Plugins/Runtime/Interchange/Source/Dispatcher/Private/InterchangeDispatcher.cpp b/Engine/Plugins/Runtime/Interchange/Source/Dispatcher/Private/InterchangeDispatcher.cpp
--- a/Engine/Plugins/Runtime/Interchange/Source/Dispatcher/Private/InterchangeDispatcher.cpp
+++ b/Engine/Plugins/Runtime/Interchange/Source/Dispatcher/Private/InterchangeDispatcher.cpp
@@ -134,9 +134,10 @@ namespace UE
 
 		void FInterchangeDispatcher::WaitAllTaskToCompleteExecution()
 		{
-			if (WorkerHandler.IsValid())
+			if (!WorkerHandler.IsValid())
 			{
 				UE_LOG(LogInterchangeDispatcher, Error, TEXT("Cannot execute tasks before starting the process"));
+				return;
 			}
 
 			bool bLogRestartError = true;
@@ -159,7 +160,7 @@ namespace UE
 			}
 			else
 			{
-				UE_LOG(LogInterchangeDispatcher, Display, TEXT("Multi Process ended and consumed all the tasks"));
+				UE_LOG(LogInterchangeDispatcher, Verbose, TEXT("Multi Process ended and consumed all the tasks"));
 			}
 		}
 
