Make sure Interchange runtime plugin description is ok
#rnx

diff --git a/Engine/Plugins/Runtime/Interchange/Interchange.uplugin b/Engine/Plugins/Runtime/Interchange/Interchange.uplugin
--- a/Engine/Plugins/Runtime/Interchange/Interchange.uplugin
+++ b/Engine/Plugins/Runtime/Interchange/Interchange.uplugin
@@ -3,7 +3,7 @@
 	"Version": 1,
 	"VersionName": "1.0",
 	"FriendlyName": "Interchange Framework",
-	"Description": "Interchange framework plugin implement all implementation of engine importers and exporters.",
+	"Description": "Interchange framework plugin implement all engine importers and exporters.",
 	"Category": "Importers",
 	"CreatedBy": "Epic Games, Inc.",
 	"CreatedByURL": "http://epicgames.com",
@@ -14,6 +14,7 @@
 	"CanContainContent": false,
 	"IsBetaVersion": false,
 	"Installed": false,
+	"IsHidden": true,
 	"Modules": [
 		{
 			"Name": "InterchangeNodes",
@@ -51,5 +52,6 @@
 			]
 		}
 	],
-	"IsExperimentalVersion": false
+	"IsExperimentalVersion": false,
+	"SupportedPrograms": [ "InterchangeWorker" ]
 }
\ No newline at end of file
