Added the capacity to import a 2d texture as a cube texutre when using interchange (for HDRI images).

diff --git a/Engine/Plugins/Runtime/Interchange/Source/Import/Private/Texture/InterchangeTextureFactory.cpp b/Engine/Plugins/Runtime/Interchange/Source/Import/Private/Texture/InterchangeTextureFactory.cpp
--- a/Engine/Plugins/Runtime/Interchange/Source/Import/Private/Texture/InterchangeTextureFactory.cpp
+++ b/Engine/Plugins/Runtime/Interchange/Source/Import/Private/Texture/InterchangeTextureFactory.cpp
@@ -108,9 +108,16 @@ namespace UE::Interchange::Private::InterchangeTextureFactory
 		{
 			if ((*TextureCubeFactoryNode)->GetCustomTranslatedTextureNodeUid(TextureNodeUniqueID))
 			{
-				if (const UInterchangeTextureCubeNode* TextureCubeTranslatedNode = Cast<UInterchangeTextureCubeNode>(NodeContainer->GetNode(TextureNodeUniqueID)))
+				if (const UInterchangeBaseNode* Node = NodeContainer->GetNode(TextureNodeUniqueID))
 				{
-					return FTextureNodeVariant(TInPlaceType<const UInterchangeTextureCubeNode*>(), TextureCubeTranslatedNode);
+					if (const UInterchangeTextureCubeNode* TextureCubeTranslatedNode = Cast<UInterchangeTextureCubeNode>(Node))
+					{
+						return FTextureNodeVariant(TInPlaceType<const UInterchangeTextureCubeNode*>(), TextureCubeTranslatedNode);
+					}
+					else if (const UInterchangeTextureNode* TextureTranslatedNode = Cast<UInterchangeTextureNode>(Node))
+					{
+						return FTextureNodeVariant(TInPlaceType<const UInterchangeTextureNode*>(), TextureTranslatedNode);
+					}
 				}
 			}
 		}
@@ -178,9 +185,9 @@ namespace UE::Interchange::Private::InterchangeTextureFactory
 	}
 
 #if WITH_EDITORONLY_DATA
-	void SetupTexture2DSourceData(UTexture2D* Texture2D, const FImportImage& Image)
+	void SetupTextureSourceData(UTexture* Texture, const FImportImage& Image)
 	{
-		Texture2D->Source.Init(
+		Texture->Source.Init(
 			Image.SizeX,
 			Image.SizeY,
 			/*NumSlices=*/ 1,
@@ -189,14 +196,14 @@ namespace UE::Interchange::Private::InterchangeTextureFactory
 			Image.RawData.GetData()
 		);
 
-		Texture2D->CompressionSettings = Image.CompressionSettings;
-		Texture2D->SRGB = Image.SRGB;
+		Texture->CompressionSettings = Image.CompressionSettings;
+		Texture->SRGB = Image.SRGB;
 
 		//If the MipGenSettings was set by the translator, we must apply it before the build
 		if (Image.MipGenSettings.IsSet())
 		{
 			// if the source has mips we keep the mips by default, unless the user changes that
-			Texture2D->MipGenSettings = Image.MipGenSettings.GetValue();
+			Texture->MipGenSettings = Image.MipGenSettings.GetValue();
 		}
 	}
 
@@ -243,7 +250,7 @@ namespace UE::Interchange::Private::InterchangeTextureFactory
 			//Import as a normal texture
 			if (BlockedImage.ImagesData.Num() == 1)
 			{
-				SetupTexture2DSourceData(Texture2D, BlockedImage.ImagesData[0]);
+				SetupTextureSourceData(Texture2D, BlockedImage.ImagesData[0]);
 				return true;
 			}
 		}
@@ -264,7 +271,7 @@ namespace UE::Interchange::Private::InterchangeTextureFactory
 		{
 			if (Image->IsSet())
 			{
-				SetupTexture2DSourceData(Texture2D, Image->GetValue());
+				SetupTextureSourceData(Texture2D, Image->GetValue());
 				return true;
 			}
 		}
@@ -328,6 +335,16 @@ namespace UE::Interchange::Private::InterchangeTextureFactory
 				}
 			}
 		}
+		else if (const TOptional<FImportImage>* OptionalImage = TexturePayload.TryGet<TOptional<FImportImage>>())
+		{
+			if (OptionalImage->IsSet())
+			{
+				const FImportImage& Image = OptionalImage->GetValue();
+				SetupTextureSourceData(TextureCube, Image);
+
+				return true;
+			}
+		}
 
 		return false;
 	}
@@ -509,7 +526,7 @@ UObject* UInterchangeTextureFactory::CreateAsset(const UInterchangeTextureFactor
 	}
 	else
 	{
-		//The texture is not a UTexture2D
+		//The texture is not supported
 		Texture->RemoveFromRoot();
 		Texture->MarkPendingKill();
 		return nullptr;
@@ -536,7 +553,7 @@ UObject* UInterchangeTextureFactory::CreateAsset(const UInterchangeTextureFactor
 		//Apply reimport strategy
 		UE::Interchange::FFactoryCommon::ApplyReimportStrategyToAsset(Arguments.ReimportStrategyFlags, Texture, PreviousNode, CurrentNode, TextureFactoryNode);
 	}
-		
+
 	//Getting the file Hash will cache it into the source data
 	Arguments.SourceData->GetFileContentHash();
 
diff --git a/Engine/Source/Developer/TextureCompressor/Private/TextureCompressorModule.cpp b/Engine/Source/Developer/TextureCompressor/Private/TextureCompressorModule.cpp
--- a/Engine/Source/Developer/TextureCompressor/Private/TextureCompressorModule.cpp
+++ b/Engine/Source/Developer/TextureCompressor/Private/TextureCompressorModule.cpp
@@ -2529,6 +2529,7 @@ private:
 			{
 				// Generate the base mip from the long-lat source image.
 				GenerateBaseCubeMipFromLongitudeLatitude2D(Mip, Image, BuildSettings.MaxTextureResolution);
+				break;
 			}
 			else
 			{
