fixes for no unity build errors
#rnx

diff --git a/Engine/Plugins/Runtime/Interchange/Source/Parsers/Fbx/Private/FbxMaterial.cpp b/Engine/Plugins/Runtime/Interchange/Source/Parsers/Fbx/Private/FbxMaterial.cpp
--- a/Engine/Plugins/Runtime/Interchange/Source/Parsers/Fbx/Private/FbxMaterial.cpp
+++ b/Engine/Plugins/Runtime/Interchange/Source/Parsers/Fbx/Private/FbxMaterial.cpp
@@ -9,6 +9,7 @@
 #include "InterchangeMaterialNode.h"
 #include "InterchangeSceneNode.h"
 #include "InterchangeTextureNode.h"
+#include "Misc/Paths.h"
 #include "Nodes/InterchangeBaseNodeContainer.h"
 
 namespace UE
diff --git a/Engine/Source/Editor/EditorWidgets/Private/ObjectNameEditSinkRegistry.cpp b/Engine/Source/Editor/EditorWidgets/Private/ObjectNameEditSinkRegistry.cpp
--- a/Engine/Source/Editor/EditorWidgets/Private/ObjectNameEditSinkRegistry.cpp
+++ b/Engine/Source/Editor/EditorWidgets/Private/ObjectNameEditSinkRegistry.cpp
@@ -1,6 +1,7 @@
 // Copyright Epic Games, Inc. All Rights Reserved.
 
 #include "ObjectNameEditSinkRegistry.h"
+#include "Editor/EditorEngine.h"
 #include "Templates/SharedPointer.h"
 #include "GameFramework/Actor.h"
 
