Make sure we support correctly the blend shape when we import skeletalmesh using interchange.

diff --git a/Engine/Plugins/Runtime/Interchange/Source/Dispatcher/Private/InterchangeCommands.cpp b/Engine/Plugins/Runtime/Interchange/Source/Dispatcher/Private/InterchangeCommands.cpp
--- a/Engine/Plugins/Runtime/Interchange/Source/Dispatcher/Private/InterchangeCommands.cpp
+++ b/Engine/Plugins/Runtime/Interchange/Source/Dispatcher/Private/InterchangeCommands.cpp
@@ -47,7 +47,7 @@ namespace UE
 		void FRunTaskCommand::SerializeImpl(FArchive& Ar)
 		{
 			Ar << JsonDescription;
-			Ar << JobIndex;
+			Ar << TaskIndex;
 		}
 
 		void FCompletedTaskCommand::SerializeImpl(FArchive& Ar)
@@ -55,6 +55,7 @@ namespace UE
 			Ar << ProcessResult;
 			Ar << JSonMessages;
 			Ar << JSonResult;
+			Ar << TaskIndex;
 		}
 	}//ns Interchange
 }//ns UE
diff --git a/Engine/Plugins/Runtime/Interchange/Source/Dispatcher/Private/InterchangeDispatcher.cpp b/Engine/Plugins/Runtime/Interchange/Source/Dispatcher/Private/InterchangeDispatcher.cpp
--- a/Engine/Plugins/Runtime/Interchange/Source/Dispatcher/Private/InterchangeDispatcher.cpp
+++ b/Engine/Plugins/Runtime/Interchange/Source/Dispatcher/Private/InterchangeDispatcher.cpp
@@ -22,13 +22,19 @@ namespace UE
 		{
 		}
 
-		int32 FInterchangeDispatcher::AddTask(const FString& InJsonDescription)
+		int32 FInterchangeDispatcher::AddTask(const FString& JsonDescription)
 		{
 			FScopeLock Lock(&TaskPoolCriticalSection);
-			int32 TaskIndex = TaskPool.Emplace(InJsonDescription);
+			int32 TaskIndex = TaskPool.Emplace(JsonDescription);
 			TaskPool[TaskIndex].Index = TaskIndex;
 			return TaskIndex;
 		}
+		int32 FInterchangeDispatcher::AddTask(const FString& JsonDescription, FInterchangeDispatcherTaskCompleted TaskCompledDelegate)
+		{
+			int32 TaskIndex = AddTask(JsonDescription);
+			TaskPool[TaskIndex].OnTaskCompleted = TaskCompledDelegate;
+			return TaskIndex;
+		}
 
 		TOptional<FTask> FInterchangeDispatcher::GetNextTask()
 		{
@@ -68,6 +74,8 @@ namespace UE
 				if (TaskState == ETaskState::ProcessOk
 					|| TaskState == ETaskState::ProcessFailed)
 				{
+					//Call the task completion delegate
+					Task.OnTaskCompleted.ExecuteIfBound(TaskIndex);
 					CompletedTaskCount++;
 				}
 
diff --git a/Engine/Plugins/Runtime/Interchange/Source/Dispatcher/Private/InterchangeDispatcherNetworking.cpp b/Engine/Plugins/Runtime/Interchange/Source/Dispatcher/Private/InterchangeDispatcherNetworking.cpp
--- a/Engine/Plugins/Runtime/Interchange/Source/Dispatcher/Private/InterchangeDispatcherNetworking.cpp
+++ b/Engine/Plugins/Runtime/Interchange/Source/Dispatcher/Private/InterchangeDispatcherNetworking.cpp
@@ -2,9 +2,9 @@
 
 #include "InterchangeDispatcherNetworking.h"
 
-#include "InterchangeDispatcherLog.h"
-
 #include "CoreMinimal.h"
+#include "InterchangeDispatcherLog.h"
+#include "Misc/ScopeLock.h"
 #include "SocketSubsystem.h"
 #include "Sockets.h"
 #include "Serialization/MemoryReader.h"
@@ -27,6 +27,7 @@ namespace UE
 
 		bool FNetworkNode::SendMessage(const TArray<uint8>& Message, double Timeout_s)
 		{
+			FScopeLock Lock(&SendReceiveCriticalSection);
 			if (bWriteError)
 			{
 				UE_LOG(LogInterchangeDispatcher, Error, TEXT("bWriteError flag raised, can't write"));
@@ -99,6 +100,7 @@ namespace UE
 
 		bool FNetworkNode::ReceiveMessage(TArray<uint8>& OutMessage, double Timeout_s)
 		{
+			FScopeLock Lock(&SendReceiveCriticalSection);
 			if (bReadError)
 			{
 				UE_LOG(LogInterchangeDispatcher, Error, TEXT("ReadError flag raised, can't read"));
diff --git a/Engine/Plugins/Runtime/Interchange/Source/Dispatcher/Private/InterchangeDispatcherTask.cpp b/Engine/Plugins/Runtime/Interchange/Source/Dispatcher/Private/InterchangeDispatcherTask.cpp
--- a/Engine/Plugins/Runtime/Interchange/Source/Dispatcher/Private/InterchangeDispatcherTask.cpp
+++ b/Engine/Plugins/Runtime/Interchange/Source/Dispatcher/Private/InterchangeDispatcherTask.cpp
@@ -53,7 +53,7 @@ namespace UE
 				return false;
 			}
 
-			if (JsonActionValue != GetAction())
+			if (!JsonActionValue.Equals(GetAction()))
 			{
 				//This json do not represent a load command
 				return false;
@@ -150,7 +150,7 @@ namespace UE
 				return false;
 			}
 
-			if (JsonActionValue != GetAction())
+			if (!JsonActionValue.Equals(GetAction()))
 			{
 				//This json do not represent a load command
 				return false;
diff --git a/Engine/Plugins/Runtime/Interchange/Source/Dispatcher/Private/InterchangeWorkerHandler.cpp b/Engine/Plugins/Runtime/Interchange/Source/Dispatcher/Private/InterchangeWorkerHandler.cpp
--- a/Engine/Plugins/Runtime/Interchange/Source/Dispatcher/Private/InterchangeWorkerHandler.cpp
+++ b/Engine/Plugins/Runtime/Interchange/Source/Dispatcher/Private/InterchangeWorkerHandler.cpp
@@ -6,6 +6,7 @@
 #include "InterchangeDispatcherConfig.h"
 #include "InterchangeDispatcherLog.h"
 
+#include "Async/TaskGraphInterfaces.h"
 #include "HAL/FileManager.h"
 #include "HAL/PlatformProcess.h"
 #include "Misc/Paths.h"
@@ -16,7 +17,46 @@ namespace UE
 {
 	namespace Interchange
 	{
+		namespace Dispatcher
+		{
+			class FTaskProcessCommand
+			{
+			private:
+				FInterchangeWorkerHandler* TaskOwner = nullptr;
+				TSharedPtr<ICommand> Command = nullptr;
+
+			public:
+				FTaskProcessCommand(FInterchangeWorkerHandler* InTaskOwner, TSharedPtr<ICommand> InCommand)
+					: TaskOwner(InTaskOwner)
+					, Command(InCommand)
+				{
+				}
+
+				static FORCEINLINE ENamedThreads::Type GetDesiredThread()
+				{
+					return ENamedThreads::AnyBackgroundThreadNormalTask;
+				}
+				static FORCEINLINE ESubsequentsMode::Type GetSubsequentsMode()
+				{
+					return ESubsequentsMode::TrackSubsequents;
+				}
+
+				FORCEINLINE TStatId GetStatId() const
+				{
+					RETURN_QUICK_DECLARE_CYCLE_STAT(FTaskProcessCommand, STATGROUP_TaskGraphTasks);
+				}
 
+				void DoTask(ENamedThreads::Type CurrentThread, const FGraphEventRef& MyCompletionGraphEvent)
+				{
+					//Verify if the task was cancel
+					if (TaskOwner == nullptr || !TaskOwner->IsAlive() || !Command.IsValid())
+					{
+						return;
+					}
+					TaskOwner->ProcessCommand(*Command);
+				}
+			};
+		}
 		static FString GetWorkerExecutablePath()
 		{
 			static FString ProcessorPath = [&]()
@@ -136,6 +176,34 @@ namespace UE
 
 		void FInterchangeWorkerHandler::RunInternal()
 		{
+
+			auto StartNewTask = [this]()->bool
+			{
+				// Fetch a new task
+				TOptional<FTask> CurrentTask = Dispatcher.GetNextTask();
+				bool bSucceed = false;
+				if (CurrentTask.IsSet())
+				{
+					FRunTaskCommand NewTask(CurrentTask.GetValue());
+					CurrentTasks.Add(NewTask.TaskIndex);
+
+					if (CommandIO.SendCommand(NewTask, Config::SendCommandTimeout_s))
+					{
+						UE_LOG(LogInterchangeDispatcher, Verbose, TEXT("New task command sent"));
+						bSucceed = true;
+					}
+					else
+					{
+						// Signal that the Task was not processed
+						TArray<FString> GarbageMessages;
+						Dispatcher.SetTaskState(CurrentTask->Index, ETaskState::UnTreated, FString(), GarbageMessages);
+						WorkerState = EWorkerState::Closing;
+						ErrorState = EWorkerErrorState::ConnectionLost_SendFailed;
+					}
+				}
+				return bSucceed;
+			};
+
 			while (IsAlive())
 			{
 				switch (WorkerState)
@@ -169,81 +237,41 @@ namespace UE
 							break;
 						}
 
-						WorkerState = EWorkerState::Idle;
+						WorkerState = EWorkerState::Processing;
 						break;
 					}
 
-					case EWorkerState::Idle:
+					case EWorkerState::Processing:
 					{
-						// Fetch a new task
-						ensureMsgf(CurrentTask.IsSet() == false, TEXT("We should not have a current task when fetching a new one"));
-						CurrentTask = Dispatcher.GetNextTask();
-
-						if (CurrentTask.IsSet())
+						ValidateConnection();
+						if (ErrorState == EWorkerErrorState::WorkerProcess_Lost)
 						{
-							FRunTaskCommand NewTask(CurrentTask.GetValue());
-
-							if (CommandIO.SendCommand(NewTask, Config::SendCommandTimeout_s))
+							for (const int32 TaskIndex : CurrentTasks)
 							{
-								UE_LOG(LogInterchangeDispatcher, Verbose, TEXT("New task command sent"));
-								WorkerState = EWorkerState::Processing;
-							}
-							else
-							{
-								// Signal that the Task was not processed
 								TArray<FString> GarbageMessages;
-								Dispatcher.SetTaskState(CurrentTask->Index, ETaskState::UnTreated, FString(), GarbageMessages);
-
-								UE_LOG(LogInterchangeDispatcher, Error, TEXT("New task command issue"));
-								WorkerState = EWorkerState::Closing;
-								ErrorState = EWorkerErrorState::ConnectionLost_SendFailed;
+								Dispatcher.SetTaskState(TaskIndex, ETaskState::ProcessFailed, FString(), GarbageMessages);
 							}
-						}
-						else if (bShouldTerminate)
-						{
-							UE_LOG(LogInterchangeDispatcher, Verbose, TEXT("Exit loop gracefully"));
+							CurrentTasks.Empty();
 							WorkerState = EWorkerState::Closing;
 						}
-						else
+						//Start a task
+						if (StartNewTask())
 						{
-							ValidateConnection();
-
-							// consume
-							if (TSharedPtr<ICommand> Command = CommandIO.GetNextCommand(Config::IdleLoopDelay))
-							{
-								ProcessCommand(*Command);
-							}
+							WorkerState = EWorkerState::Processing;
 						}
 
-						break;
-					}
-
-					case EWorkerState::Processing:
-					{
-						if (TSharedPtr<ICommand> Command = CommandIO.GetNextCommand(Config::ProcessingLoopDelay))
+						// consume task
+						if (TSharedPtr<ICommand> Command = CommandIO.GetNextCommand(Config::IdleLoopDelay))
 						{
-							ProcessCommand(*Command);
-
-							bool bProcessingOver = CurrentTask.IsSet() == false;
-							if (bProcessingOver)
-							{
-								WorkerState = bShouldTerminate ? EWorkerState::Closing : EWorkerState::Idle;
-							}
+							//Make the processing asynchronous
+							TGraphTask<UE::Interchange::Dispatcher::FTaskProcessCommand>::CreateTask().ConstructAndDispatchWhenReady(this, Command);
 						}
-						else
-						{
-							ValidateConnection();
-							if (ErrorState == EWorkerErrorState::WorkerProcess_Lost)
-							{
-								if (CurrentTask.IsSet())
-								{
-									TArray<FString> GarbageMessages;
-									Dispatcher.SetTaskState(CurrentTask->Index, ETaskState::ProcessFailed, FString(), GarbageMessages);
-									CurrentTask.Reset();
-								}
 
-								WorkerState = EWorkerState::Closing;
-							}
+						//Need to terminate?
+						if (bShouldTerminate)
+						{
+							UE_LOG(LogInterchangeDispatcher, Verbose, TEXT("Exit loop gracefully"));
+							WorkerState = EWorkerState::Closing;
 						}
 						break;
 					}
@@ -341,13 +369,9 @@ namespace UE
 
 		void FInterchangeWorkerHandler::ProcessCommand(FCompletedTaskCommand& CompletedTaskCommand)
 		{
-			if (!CurrentTask.IsSet())
-			{
-				return;
-			}
 
-			Dispatcher.SetTaskState(CurrentTask->Index, CompletedTaskCommand.ProcessResult, CompletedTaskCommand.JSonResult, CompletedTaskCommand.JSonMessages);
-			CurrentTask.Reset();
+			Dispatcher.SetTaskState(CompletedTaskCommand.TaskIndex, CompletedTaskCommand.ProcessResult, CompletedTaskCommand.JSonResult, CompletedTaskCommand.JSonMessages);
+			CurrentTasks.Remove(CompletedTaskCommand.TaskIndex);
 		}
 
 		const TCHAR* FInterchangeWorkerHandler::EWorkerErrorStateAsString(EWorkerErrorState Error)
diff --git a/Engine/Plugins/Runtime/Interchange/Source/Dispatcher/Public/InterchangeCommands.h b/Engine/Plugins/Runtime/Interchange/Source/Dispatcher/Public/InterchangeCommands.h
--- a/Engine/Plugins/Runtime/Interchange/Source/Dispatcher/Public/InterchangeCommands.h
+++ b/Engine/Plugins/Runtime/Interchange/Source/Dispatcher/Public/InterchangeCommands.h
@@ -121,7 +121,7 @@ namespace UE
 		{
 		public:
 			FRunTaskCommand() = default;
-			FRunTaskCommand(const FTask& Task) : JsonDescription(Task.JsonDescription), JobIndex(Task.Index) {}
+			FRunTaskCommand(const FTask& Task) : JsonDescription(Task.JsonDescription), TaskIndex(Task.Index) {}
 			virtual ECommandId GetType() const override { return ECommandId::RunTask; }
 
 		protected:
@@ -129,7 +129,7 @@ namespace UE
 
 		public:
 			FString JsonDescription;
-			int32 JobIndex = -1;
+			int32 TaskIndex = -1;
 		};
 
 		class INTERCHANGEDISPATCHER_API FCompletedTaskCommand : public ICommand
@@ -144,6 +144,7 @@ namespace UE
 			ETaskState ProcessResult = ETaskState::Unknown;
 			FString JSonResult;
 			TArray<FString> JSonMessages;
+			int32 TaskIndex = INDEX_NONE;
 		};
 
 	} //ns Interchange
diff --git a/Engine/Plugins/Runtime/Interchange/Source/Dispatcher/Public/InterchangeDispatcher.h b/Engine/Plugins/Runtime/Interchange/Source/Dispatcher/Public/InterchangeDispatcher.h
--- a/Engine/Plugins/Runtime/Interchange/Source/Dispatcher/Public/InterchangeDispatcher.h
+++ b/Engine/Plugins/Runtime/Interchange/Source/Dispatcher/Public/InterchangeDispatcher.h
@@ -20,6 +20,7 @@ namespace UE
 			~FInterchangeDispatcher() { TerminateProcess(); }
 
 			int32 AddTask(const FString& JsonDescription);
+			int32 AddTask(const FString& JsonDescription, FInterchangeDispatcherTaskCompleted TaskCompledDelegate);
 			TOptional<FTask> GetNextTask();
 			void SetTaskState(int32 TaskIndex, ETaskState TaskState, const FString& JsonResult, const TArray<FString>& JSonMessages);
 			void GetTaskState(int32 TaskIndex, ETaskState& TaskState, FString& JsonResult, TArray<FString>& JSonMessages);
diff --git a/Engine/Plugins/Runtime/Interchange/Source/Dispatcher/Public/InterchangeDispatcherNetworking.h b/Engine/Plugins/Runtime/Interchange/Source/Dispatcher/Public/InterchangeDispatcherNetworking.h
--- a/Engine/Plugins/Runtime/Interchange/Source/Dispatcher/Public/InterchangeDispatcherNetworking.h
+++ b/Engine/Plugins/Runtime/Interchange/Source/Dispatcher/Public/InterchangeDispatcherNetworking.h
@@ -6,6 +6,7 @@
 
 #include "CoreMinimal.h"
 #include "Containers/Queue.h"
+#include "HAL/CriticalSection.h"
 #include "Sockets.h"
 
 namespace UE
@@ -74,6 +75,8 @@ namespace UE
 
 			bool bReadError = false;
 			bool bWriteError = false;
+
+			FCriticalSection SendReceiveCriticalSection;
 		};
 
 
diff --git a/Engine/Plugins/Runtime/Interchange/Source/Dispatcher/Public/InterchangeDispatcherTask.h b/Engine/Plugins/Runtime/Interchange/Source/Dispatcher/Public/InterchangeDispatcherTask.h
--- a/Engine/Plugins/Runtime/Interchange/Source/Dispatcher/Public/InterchangeDispatcherTask.h
+++ b/Engine/Plugins/Runtime/Interchange/Source/Dispatcher/Public/InterchangeDispatcherTask.h
@@ -3,6 +3,16 @@
 #pragma once
 
 #include "CoreMinimal.h"
+
+namespace UE
+{
+	namespace Interchange
+	{
+		struct FTask;
+	}
+}
+DECLARE_DELEGATE_OneParam(FInterchangeDispatcherTaskCompleted, int32 TaskIndex);
+
 namespace UE
 {
 	namespace Interchange
@@ -18,7 +28,7 @@ namespace UE
 
 		struct FTask
 		{
-			FTask() = default;
+			FTask() = delete;
 
 			FTask(const FString& InJsonDescription)
 			{
@@ -31,6 +41,7 @@ namespace UE
 			ETaskState State = ETaskState::Unknown;
 			FString JsonResult;
 			TArray<FString> JsonMessages;
+			FInterchangeDispatcherTaskCompleted OnTaskCompleted;
 		};
 
 		/**
@@ -88,7 +99,7 @@ namespace UE
 
 			virtual FString GetAction() const override
 			{
-				static const FString LoadString = TEXT("Load");
+				static const FString LoadString = TEXT("LoadSource");
 				return LoadString;
 			}
 
diff --git a/Engine/Plugins/Runtime/Interchange/Source/Dispatcher/Public/InterchangeWorkerHandler.h b/Engine/Plugins/Runtime/Interchange/Source/Dispatcher/Public/InterchangeWorkerHandler.h
--- a/Engine/Plugins/Runtime/Interchange/Source/Dispatcher/Public/InterchangeWorkerHandler.h
+++ b/Engine/Plugins/Runtime/Interchange/Source/Dispatcher/Public/InterchangeWorkerHandler.h
@@ -15,6 +15,11 @@ namespace UE
 {
 	namespace Interchange
 	{
+		namespace Dispatcher
+		{
+			class FTaskProcessCommand;
+		}
+
 		class FInterchangeDispatcher;
 
 		//Handle a Worker by socket communication
@@ -23,8 +28,7 @@ namespace UE
 			enum class EWorkerState
 			{
 				Uninitialized,
-				Idle, // Initialized, available for processing
-				Processing, // Currently processing a task
+				Processing, // send task and receive result
 				Closing, // in the process of terminating
 				Terminated, // aka. Not Alive
 			};
@@ -49,12 +53,13 @@ namespace UE
 			void Stop();
 			void StopBlocking();
 
+		protected:
+			void ProcessCommand(ICommand& Command);
 		private:
 			void RunInternal();
 			void StartWorkerProcess();
 			void ValidateConnection();
 
-			void ProcessCommand(ICommand& Command);
 			void ProcessCommand(FPingCommand& PingCommand);
 			void ProcessCommand(FCompletedTaskCommand& RunTaskCommand);
 			const TCHAR* EWorkerErrorStateAsString(EWorkerErrorState e);
@@ -75,9 +80,10 @@ namespace UE
 
 			// self
 			FString ResultFolder;
-			TOptional<FTask> CurrentTask;
+			TArray<int32> CurrentTasks;
 			bool bShouldTerminate;
 
+			friend Dispatcher::FTaskProcessCommand;
 		};
 	} //ns Interchange
 }//ns UE
diff --git a/Engine/Plugins/Runtime/Interchange/Source/Import/Private/Fbx/InterchangeFbxTranslator.cpp b/Engine/Plugins/Runtime/Interchange/Source/Import/Private/Fbx/InterchangeFbxTranslator.cpp
--- a/Engine/Plugins/Runtime/Interchange/Source/Import/Private/Fbx/InterchangeFbxTranslator.cpp
+++ b/Engine/Plugins/Runtime/Interchange/Source/Import/Private/Fbx/InterchangeFbxTranslator.cpp
@@ -152,121 +152,121 @@ TOptional<UE::Interchange::FStaticMeshPayloadData> UInterchangeFbxTranslator::Ge
 	return TOptional<UE::Interchange::FStaticMeshPayloadData>();
 }
 
-TOptional<UE::Interchange::FSkeletalMeshLodPayloadData> UInterchangeFbxTranslator::GetSkeletalMeshLodPayloadData(const FString& PayLoadKey) const
+void UInterchangeFbxTranslator::GetSkeletalMeshLodPayloadData(const FString& PayLoadKey, TOptional<UE::Interchange::FSkeletalMeshLodPayloadData>& OptionalSkeletalMeshLodPayloadData) const
 {
 	if (!Dispatcher.IsValid())
 	{
-		return TOptional<UE::Interchange::FSkeletalMeshLodPayloadData>();
+		return;
 	}
 
 	//Create a json command to read the fbx file
 	FString JsonCommand = CreateFetchPayloadFbxCommand(PayLoadKey);
-	int32 TaskIndex = Dispatcher->AddTask(JsonCommand);
-
-	//Blocking call until all tasks are executed
-	Dispatcher->WaitAllTaskToCompleteExecution();
-
-	UE::Interchange::ETaskState TaskState;
-	FString JsonResult;
-	TArray<FString> JSonMessages;
-	Dispatcher->GetTaskState(TaskIndex, TaskState, JsonResult, JSonMessages);
-
-	//TODO: Parse the JSonMessage and add the message to the interchange not yet develop error messaging
-
-	if (TaskState != UE::Interchange::ETaskState::ProcessOk)
+	int32 TaskIndex = Dispatcher->AddTask(JsonCommand, FInterchangeDispatcherTaskCompleted::CreateLambda([this, &OptionalSkeletalMeshLodPayloadData](const int32 TaskIndex)
 	{
-		return TOptional<UE::Interchange::FSkeletalMeshLodPayloadData>();
-	}
-	//Grab the result file and fill the BaseNodeContainer
-	UE::Interchange::FJsonFetchPayloadCmd::JsonResultParser ResultParser;
-	ResultParser.FromJson(JsonResult);
-	FString SkeletalMeshPayloadFilename = ResultParser.GetResultFilename();
+		UE::Interchange::ETaskState TaskState;
+		FString JsonResult;
+		TArray<FString> JSonMessages;
+		Dispatcher->GetTaskState(TaskIndex, TaskState, JsonResult, JSonMessages);
 
-	if (!ensure(FPaths::FileExists(SkeletalMeshPayloadFilename)))
-	{
-		//TODO log an error saying the payload file do not exist even if the get payload command succeed
-		return TOptional<UE::Interchange::FSkeletalMeshLodPayloadData>();
-	}
-	UE::Interchange::FSkeletalMeshLodPayloadData SkeletalMeshLodPayload;
-	SkeletalMeshLodPayload.LodMeshDescription.Empty();
-	//All sub object should be gone with the reset
-	TArray64<uint8> Buffer;
-	FFileHelper::LoadFileToArray(Buffer, *SkeletalMeshPayloadFilename);
-	uint8* FileData = Buffer.GetData();
-	int64 FileDataSize = Buffer.Num();
-	if (FileDataSize < 1)
-	{
-		//Nothing to load from this file
-		return TOptional<UE::Interchange::FSkeletalMeshLodPayloadData>();
-	}
+		//TODO: Parse the JSonMessage and add the message to the interchange not yet develop error messaging
+		if (TaskState != UE::Interchange::ETaskState::ProcessOk)
+		{
+			return;
+		}
+		//Grab the result file and fill the BaseNodeContainer
+		UE::Interchange::FJsonFetchPayloadCmd::JsonResultParser ResultParser;
+		ResultParser.FromJson(JsonResult);
+		FString SkeletalMeshPayloadFilename = ResultParser.GetResultFilename();
 
-	//Buffer keep the ownership of the data, the large memory reader is use to serialize the TMap
-	FLargeMemoryReader Ar(FileData, FileDataSize);
-	SkeletalMeshLodPayload.LodMeshDescription.Serialize(Ar);
-	bool bFetchSkinnedData = false;
-	Ar << bFetchSkinnedData;
-	if (bFetchSkinnedData)
-	{
-		//Read the bone Name to remap the influence correctly
-		Ar << SkeletalMeshLodPayload.JointNames;
-	}
+		if (!ensure(FPaths::FileExists(SkeletalMeshPayloadFilename)))
+		{
+			//TODO log an error saying the payload file do not exist even if the get payload command succeed
+			return;
+		}
+		UE::Interchange::FSkeletalMeshLodPayloadData SkeletalMeshLodPayload;
+		SkeletalMeshLodPayload.LodMeshDescription.Empty();
+		//All sub object should be gone with the reset
+		TArray64<uint8> Buffer;
+		FFileHelper::LoadFileToArray(Buffer, *SkeletalMeshPayloadFilename);
+		uint8* FileData = Buffer.GetData();
+		int64 FileDataSize = Buffer.Num();
+		if (FileDataSize < 1)
+		{
+			//Nothing to load from this file
+			return;
+		}
 
-	return SkeletalMeshLodPayload;
+		//Buffer keep the ownership of the data, the large memory reader is use to serialize the TMap
+		FLargeMemoryReader Ar(FileData, FileDataSize);
+		SkeletalMeshLodPayload.LodMeshDescription.Serialize(Ar);
+		bool bFetchSkinnedData = false;
+		Ar << bFetchSkinnedData;
+		if (bFetchSkinnedData)
+		{
+			//Read the bone Name to remap the influence correctly
+			Ar << SkeletalMeshLodPayload.JointNames;
+		}
+		OptionalSkeletalMeshLodPayloadData = SkeletalMeshLodPayload;
+	}));
 }
 
-TOptional<UE::Interchange::FSkeletalMeshBlendShapePayloadData> UInterchangeFbxTranslator::GetSkeletalMeshBlendShapePayloadData(const FString& PayLoadKey) const
+void UInterchangeFbxTranslator::GetSkeletalMeshBlendShapePayloadData(const FString& PayLoadKey, TOptional<UE::Interchange::FSkeletalMeshBlendShapePayloadData>& OptionalSkeletalMeshBlendShapePayloadData) const
 {
 	if (!Dispatcher.IsValid())
 	{
-		return TOptional<UE::Interchange::FSkeletalMeshBlendShapePayloadData>();
+		return;
 	}
 
 	//Create a json command to read the fbx file
 	FString JsonCommand = CreateFetchPayloadFbxCommand(PayLoadKey);
-	int32 TaskIndex = Dispatcher->AddTask(JsonCommand);
-
-	//Blocking call until all tasks are executed
-	Dispatcher->WaitAllTaskToCompleteExecution();
 
-	UE::Interchange::ETaskState TaskState;
-	FString JsonResult;
-	TArray<FString> JSonMessages;
-	Dispatcher->GetTaskState(TaskIndex, TaskState, JsonResult, JSonMessages);
+	const int32 CreatedTaskIndex = Dispatcher->AddTask(JsonCommand,	FInterchangeDispatcherTaskCompleted::CreateLambda([this, &OptionalSkeletalMeshBlendShapePayloadData](const int32 TaskIndex)
+	{
+		UE::Interchange::ETaskState TaskState;
+		FString JsonResult;
+		TArray<FString> JSonMessages;
+		Dispatcher->GetTaskState(TaskIndex, TaskState, JsonResult, JSonMessages);
 
-	//TODO: Parse the JSonMessage and add the message to the interchange not yet develop error messaging
+		//TODO: Parse the JSonMessage and add the message to the interchange not yet develop error messaging
 
-	if (TaskState != UE::Interchange::ETaskState::ProcessOk)
-	{
-		return TOptional<UE::Interchange::FSkeletalMeshBlendShapePayloadData>();
-	}
-	//Grab the result file and fill the BaseNodeContainer
-	UE::Interchange::FJsonFetchPayloadCmd::JsonResultParser ResultParser;
-	ResultParser.FromJson(JsonResult);
-	FString SkeletalMeshPayloadFilename = ResultParser.GetResultFilename();
+		if (TaskState != UE::Interchange::ETaskState::ProcessOk)
+		{
+			return;
+		}
+		//Grab the result file and fill the BaseNodeContainer
+		UE::Interchange::FJsonFetchPayloadCmd::JsonResultParser ResultParser;
+		ResultParser.FromJson(JsonResult);
+		FString SkeletalMeshPayloadFilename = ResultParser.GetResultFilename();
 
-	if (!ensure(FPaths::FileExists(SkeletalMeshPayloadFilename)))
-	{
-		//TODO log an error saying the payload file do not exist even if the get payload command succeed
-		return TOptional<UE::Interchange::FSkeletalMeshBlendShapePayloadData>();
-	}
-	UE::Interchange::FSkeletalMeshBlendShapePayloadData SkeletalMeshBlendShapePayload;
-	SkeletalMeshBlendShapePayload.LodMeshDescription.Empty();
-	//All sub object should be gone with the reset
-	TArray64<uint8> Buffer;
-	FFileHelper::LoadFileToArray(Buffer, *SkeletalMeshPayloadFilename);
-	uint8* FileData = Buffer.GetData();
-	int64 FileDataSize = Buffer.Num();
-	if (FileDataSize < 1)
-	{
-		//Nothing to load from this file
-		return TOptional<UE::Interchange::FSkeletalMeshBlendShapePayloadData>();
-	}
+		if (!ensure(FPaths::FileExists(SkeletalMeshPayloadFilename)))
+		{
+			//TODO log an error saying the payload file do not exist even if the get payload command succeed
+			return;
+		}
+		UE::Interchange::FSkeletalMeshBlendShapePayloadData SkeletalMeshBlendShapePayload;
+		SkeletalMeshBlendShapePayload.LodMeshDescription.Empty();
+		//All sub object should be gone with the reset
+		TArray64<uint8> Buffer;
+		FFileHelper::LoadFileToArray(Buffer, *SkeletalMeshPayloadFilename);
+		uint8* FileData = Buffer.GetData();
+		int64 FileDataSize = Buffer.Num();
+		if (FileDataSize < 1)
+		{
+			//Nothing to load from this file
+			return;
+		}
 
-	//Buffer keep the ownership of the data, the large memory reader is use to serialize the TMap
-	FLargeMemoryReader Ar(FileData, FileDataSize);
-	SkeletalMeshBlendShapePayload.LodMeshDescription.Serialize(Ar);
+		//Buffer keep the ownership of the data, the large memory reader is use to serialize the TMap
+		FLargeMemoryReader Ar(FileData, FileDataSize);
+		SkeletalMeshBlendShapePayload.LodMeshDescription.Serialize(Ar);
+		OptionalSkeletalMeshBlendShapePayloadData = SkeletalMeshBlendShapePayload;
+	}));
+}
 
-	return SkeletalMeshBlendShapePayload;
+void UInterchangeFbxTranslator::WaitUntilAllSkeletalMeshPayloadCommandAreCompleted() const
+{
+	//Blocking call until all tasks are executed
+	Dispatcher->WaitAllTaskToCompleteExecution();
 }
 
 FString UInterchangeFbxTranslator::CreateLoadFbxFileCommand(const FString& FbxFilePath) const
diff --git a/Engine/Plugins/Runtime/Interchange/Source/Import/Private/Mesh/InterchangeSkeletalMeshFactory.cpp b/Engine/Plugins/Runtime/Interchange/Source/Import/Private/Mesh/InterchangeSkeletalMeshFactory.cpp
--- a/Engine/Plugins/Runtime/Interchange/Source/Import/Private/Mesh/InterchangeSkeletalMeshFactory.cpp
+++ b/Engine/Plugins/Runtime/Interchange/Source/Import/Private/Mesh/InterchangeSkeletalMeshFactory.cpp
@@ -1,6 +1,7 @@
 // Copyright Epic Games, Inc. All Rights Reserved. 
 #include "Mesh/InterchangeSkeletalMeshFactory.h"
 
+#include "Async/ParallelFor.h"
 #include "Components.h"
 #include "Engine/SkeletalMesh.h"
 #include "GPUSkinPublicDefs.h"
@@ -37,6 +38,15 @@ namespace UE
 	{
 		namespace Private
 		{
+			//Get the mesh node context for each MeshUids
+			struct FMeshNodeContext
+			{
+				const UInterchangeMeshNode* MeshNode = nullptr;
+				const UInterchangeSceneNode* SceneNode = nullptr;
+				TOptional<FTransform> SceneGlobalTransform;
+				FString TranslatorPayloadKey;
+			};
+
 			struct FJointInfo
 			{
 				FString Name;
@@ -135,9 +145,75 @@ namespace UE
 				return true;
 			}
 
+			void FillBlendShapeMeshDescriptionsPerBlendShapeName(const FMeshNodeContext& MeshNodeContext
+																 , TMap<FString, TOptional<UE::Interchange::FSkeletalMeshBlendShapePayloadData>>& BlendShapeMeshDescriptionsPerBlendShapeName
+																 , const IInterchangeSkeletalMeshPayloadInterface* SkeletalMeshTranslatorPayloadInterface
+																 , const int32 VertexOffset
+																 , const UInterchangeBaseNodeContainer* NodeContainer
+																 , FString AssetName)
+			{
+				TArray<FString> BlendShapeUids;
+				MeshNodeContext.MeshNode->GetShapeDependencies(BlendShapeUids);
+				//Make sure we do not reallocate the map since its pass to the asynchronous GetSkeletalMeshBlendShapePayloadData.
+				BlendShapeMeshDescriptionsPerBlendShapeName.Reserve(BlendShapeUids.Num());
+				for (const FString& BlendShapeUid : BlendShapeUids)
+				{
+					if (const UInterchangeMeshNode* BlendShapeMeshNode = Cast<UInterchangeMeshNode>(NodeContainer->GetNode(BlendShapeUid)))
+					{
+						TOptional<FString> BlendShapePayloadKey = BlendShapeMeshNode->GetPayLoadKey();
+						if (!BlendShapePayloadKey.IsSet())
+						{
+							UE_LOG(LogInterchangeImport, Warning, TEXT("Empty LOD morph target mesh reference payload when importing SkeletalMesh asset %s"), *AssetName);
+							continue;
+						}
+						TOptional<UE::Interchange::FSkeletalMeshBlendShapePayloadData>& BlendShapeMeshPayload = BlendShapeMeshDescriptionsPerBlendShapeName.Add(BlendShapeUid);
+						SkeletalMeshTranslatorPayloadInterface->GetSkeletalMeshBlendShapePayloadData(BlendShapePayloadKey.GetValue(), BlendShapeMeshPayload);
+					}
+				}
+
+				//Wait until all payload are loaded
+				SkeletalMeshTranslatorPayloadInterface->WaitUntilAllSkeletalMeshPayloadCommandAreCompleted();
+
+				for (const FString& BlendShapeUid : BlendShapeUids)
+				{
+					if (const UInterchangeMeshNode* BlendShapeMeshNode = Cast<UInterchangeMeshNode>(NodeContainer->GetNode(BlendShapeUid)))
+					{
+						TOptional<FString> BlendShapePayloadKey = BlendShapeMeshNode->GetPayLoadKey();
+						if (!BlendShapePayloadKey.IsSet())
+						{
+							continue;
+						}
+						if (!ensure(BlendShapeMeshDescriptionsPerBlendShapeName.Contains(BlendShapeUid)))
+						{
+							continue;
+						}
+
+						TOptional<UE::Interchange::FSkeletalMeshBlendShapePayloadData>& BlendShapeMeshPayload = BlendShapeMeshDescriptionsPerBlendShapeName.FindChecked(BlendShapeUid);
+						if (!BlendShapeMeshPayload.IsSet())
+						{
+							BlendShapeMeshDescriptionsPerBlendShapeName.Remove(BlendShapeUid);
+							UE_LOG(LogInterchangeImport, Warning, TEXT("Invalid Skeletal mesh morph target payload key [%s] SkeletalMesh asset %s"), *BlendShapePayloadKey.GetValue(), *AssetName);
+							continue;
+						}
+						BlendShapeMeshPayload->VertexOffset = VertexOffset;
+						//Use the Mesh node parent bake transform
+						if (MeshNodeContext.SceneGlobalTransform.IsSet())
+						{
+							BlendShapeMeshPayload->GlobalTransform = MeshNodeContext.SceneGlobalTransform;
+						}
+						else
+						{
+							BlendShapeMeshPayload->GlobalTransform.Reset();
+						}
+					}
+				}
+			}
 
 			bool CopyBlendShapesMeshDescriptionToSkeletalMeshImportData(const TMap<FString, TOptional<UE::Interchange::FSkeletalMeshBlendShapePayloadData>>& LodBlendShapeMeshDescriptions, FSkeletalMeshImportData& DestinationSkeletalMeshImportData)
 			{
+				const int32 OriginalMorphTargetCount = LodBlendShapeMeshDescriptions.Num();
+				TArray<FString> Keys;
+				int32 MorphTargetCount = 0;
 				for (const TPair<FString, TOptional<UE::Interchange::FSkeletalMeshBlendShapePayloadData>>& Pair : LodBlendShapeMeshDescriptions)
 				{
 					const FString BlendShapeName(Pair.Key);
@@ -148,41 +224,94 @@ namespace UE
 						continue;
 					}
 
-					FTransform GlobalTransform = BlendShapePayloadData->GlobalTransform.IsSet() ? BlendShapePayloadData->GlobalTransform.GetValue() : FTransform::Identity;
-					TArray<FVector3f> CompressPoints;
-					CompressPoints.Reserve(DestinationSkeletalMeshImportData.Points.Num());
-
 					const FMeshDescription& SourceMeshDescription = BlendShapePayloadData.GetValue().LodMeshDescription;
-					FStaticMeshConstAttributes Attributes(SourceMeshDescription);
-					TVertexAttributesConstRef<FVector3f> VertexPositions = Attributes.GetVertexPositions();
+					const int32 VertexOffset = BlendShapePayloadData->VertexOffset;
+					const int32 SourceMeshVertexCount = SourceMeshDescription.Vertices().Num();
+					const int32 DestinationVertexIndexMax = VertexOffset + SourceMeshVertexCount;
+					if (!DestinationSkeletalMeshImportData.Points.IsValidIndex(DestinationVertexIndexMax))
+					{
+						UE_LOG(LogInterchangeImport, Error, TEXT("Corrupted blend shape optional payload data [%s]"), *BlendShapeName);
+						continue;
+					}
+					Keys.Add(Pair.Key);
+					MorphTargetCount++;
+				}
+				ensure(Keys.Num() == MorphTargetCount);
+				//Allocate the data
+				DestinationSkeletalMeshImportData.MorphTargetNames.AddDefaulted(MorphTargetCount);
+				DestinationSkeletalMeshImportData.MorphTargetModifiedPoints.AddDefaulted(MorphTargetCount);
+				DestinationSkeletalMeshImportData.MorphTargets.AddDefaulted(MorphTargetCount);
+
+				int32 NumMorphGroup = FMath::Min(FPlatformMisc::NumberOfWorkerThreadsToSpawn(), MorphTargetCount);
+				const int32 MorphTargetGroupSize = FMath::Max(FMath::CeilToInt(static_cast<float>(MorphTargetCount) / static_cast<float>(NumMorphGroup)), 1);
+				ParallelFor(NumMorphGroup, [MorphTargetGroupSize,
+							MorphTargetCount,
+							NumMorphGroup,
+							&LodBlendShapeMeshDescriptions,
+							&Keys,
+							&DestinationSkeletalMeshImportData](const int32 MorphTargetGroupIndex)
+				{
+					const int32 MorphTargetIndexOffset = MorphTargetGroupIndex * MorphTargetGroupSize;
+					const int32 MorphTargetEndLoopCount = MorphTargetIndexOffset + MorphTargetGroupSize;
+					for (int32 MorphTargetIndex = MorphTargetIndexOffset; MorphTargetIndex < MorphTargetEndLoopCount; ++MorphTargetIndex)
+					{
+						if (!Keys.IsValidIndex(MorphTargetIndex))
+						{
+							ensure(MorphTargetGroupIndex + 1 == NumMorphGroup);
+							//Executing the last morph target group, and we do not have a full group.
+							break;
+						}
+						const FString BlendShapeName(Keys[MorphTargetIndex]);
+						const TOptional<UE::Interchange::FSkeletalMeshBlendShapePayloadData>& BlendShapePayloadData = LodBlendShapeMeshDescriptions.FindChecked(BlendShapeName);
+						if (!ensure(BlendShapePayloadData.IsSet()))
+						{
+							//This error was suppose to be catch in the pre parallel for loop
+							break;
+						}
 
-					//Create the morph target source data
-					FString& MorphTargetName = DestinationSkeletalMeshImportData.MorphTargetNames.AddDefaulted_GetRef();
-					MorphTargetName = Pair.Key;
-					TSet<uint32>& ModifiedPoints = DestinationSkeletalMeshImportData.MorphTargetModifiedPoints.AddDefaulted_GetRef();
-					FSkeletalMeshImportData& MorphTargetData = DestinationSkeletalMeshImportData.MorphTargets.AddDefaulted_GetRef();
+						const FMeshDescription& SourceMeshDescription = BlendShapePayloadData.GetValue().LodMeshDescription;
+						const FTransform GlobalTransform = BlendShapePayloadData->GlobalTransform.IsSet() ? BlendShapePayloadData->GlobalTransform.GetValue() : FTransform::Identity;
+						const int32 VertexOffset = BlendShapePayloadData->VertexOffset;
+						const int32 SourceMeshVertexCount = SourceMeshDescription.Vertices().Num();
+						const int32 DestinationVertexIndexMax = VertexOffset + SourceMeshVertexCount;
+						if (!ensure(DestinationSkeletalMeshImportData.Points.IsValidIndex(DestinationVertexIndexMax)))
+						{
+							//This error was suppose to be catch in the pre parallel for loop
+							break;
+						}
+						TArray<FVector3f> CompressPoints;
+						CompressPoints.Reserve(SourceMeshVertexCount);
+						FStaticMeshConstAttributes Attributes(SourceMeshDescription);
+						TVertexAttributesConstRef<FVector3f> VertexPositions = Attributes.GetVertexPositions();
 
-					//Reserve the point and influences
-					MorphTargetData.Points.AddZeroed(SourceMeshDescription.Vertices().Num());
+						//Create the morph target source data
+						FString& MorphTargetName = DestinationSkeletalMeshImportData.MorphTargetNames[MorphTargetIndex];
+						MorphTargetName = BlendShapeName;
+						TSet<uint32>& ModifiedPoints = DestinationSkeletalMeshImportData.MorphTargetModifiedPoints[MorphTargetIndex];
+						FSkeletalMeshImportData& MorphTargetData = DestinationSkeletalMeshImportData.MorphTargets[MorphTargetIndex];
 
-					for (FVertexID VertexID : SourceMeshDescription.Vertices().GetElementIDs())
-					{
-						//We can use GetValue because the Meshdescription was compacted before the copy
-						MorphTargetData.Points[VertexID.GetValue()] = GlobalTransform.TransformPosition(VertexPositions[VertexID]);
-					}
+						//Reserve the point and influences
+						MorphTargetData.Points.AddZeroed(SourceMeshVertexCount);
 
-					for (int32 PointIdx = 0; PointIdx < DestinationSkeletalMeshImportData.Points.Num(); ++PointIdx)
-					{
-						int32 OriginalPointIdx = DestinationSkeletalMeshImportData.PointToRawMap[PointIdx];
-						//Rebuild the data with only the modified point
-						if ((MorphTargetData.Points[OriginalPointIdx] - DestinationSkeletalMeshImportData.Points[PointIdx]).SizeSquared() > FMath::Square(THRESH_POINTS_ARE_SAME))
+						for (FVertexID VertexID : SourceMeshDescription.Vertices().GetElementIDs())
 						{
-							ModifiedPoints.Add(PointIdx);
-							CompressPoints.Add(MorphTargetData.Points[OriginalPointIdx]);
+							//We can use GetValue because the Meshdescription was compacted before the copy
+							MorphTargetData.Points[VertexID.GetValue()] = GlobalTransform.TransformPosition(VertexPositions[VertexID]);
 						}
+
+						for (int32 PointIdx = VertexOffset; PointIdx < DestinationVertexIndexMax; ++PointIdx)
+						{
+							int32 OriginalPointIdx = DestinationSkeletalMeshImportData.PointToRawMap[PointIdx] - VertexOffset;
+							//Rebuild the data with only the modified point
+							if ((MorphTargetData.Points[OriginalPointIdx] - DestinationSkeletalMeshImportData.Points[PointIdx]).SizeSquared() > FMath::Square(THRESH_POINTS_ARE_SAME))
+							{
+								ModifiedPoints.Add(PointIdx);
+								CompressPoints.Add(MorphTargetData.Points[OriginalPointIdx]);
+							}
+						}
+						MorphTargetData.Points = CompressPoints;
 					}
-					MorphTargetData.Points = CompressPoints;
-				}
+				});
 				return true;
 			}
 
@@ -558,15 +687,8 @@ UObject* UInterchangeSkeletalMeshFactory::CreateAsset(const UInterchangeSkeletal
 					continue;
 				}
 
-				//Get the mesh node context for each MeshUids
-				struct FMeshNodeContext
-				{
-					const UInterchangeMeshNode* MeshNode = nullptr;
-					const UInterchangeSceneNode* SceneNode = nullptr;
-					TOptional<FTransform> SceneGlobalTransform;
-					FString TranslatorPayloadKey;
-				};
-				TArray<FMeshNodeContext> MeshReferences;
+				
+				TArray<UE::Interchange::Private::FMeshNodeContext> MeshReferences;
 				//Scope to query the mesh node
 				{
 					TArray<FString> MeshUids;
@@ -574,7 +696,7 @@ UObject* UInterchangeSkeletalMeshFactory::CreateAsset(const UInterchangeSkeletal
 					MeshReferences.Reserve(MeshUids.Num());
 					for (const FString& MeshUid : MeshUids)
 					{
-						FMeshNodeContext MeshReference;
+						UE::Interchange::Private::FMeshNodeContext MeshReference;
 						MeshReference.MeshNode = Cast<UInterchangeMeshNode>(Arguments.NodeContainer->GetNode(MeshUid));
 						if (!MeshReference.MeshNode)
 						{
@@ -666,10 +788,21 @@ UObject* UInterchangeSkeletalMeshFactory::CreateAsset(const UInterchangeSkeletal
 				//Blend shape MeshDescriptions
 				TMap<FString, TOptional<UE::Interchange::FSkeletalMeshBlendShapePayloadData>> BlendShapeMeshDescriptionsPerBlendShapeName;
 
+				TMap<FString, TOptional<UE::Interchange::FSkeletalMeshLodPayloadData>> LodMeshPayloadPerTranslatorPayloadKey;
+				LodMeshPayloadPerTranslatorPayloadKey.Reserve(MeshReferences.Num());
+				for (const UE::Interchange::Private::FMeshNodeContext& MeshNodeContext : MeshReferences)
+				{
+					TOptional<UE::Interchange::FSkeletalMeshLodPayloadData>& LodMeshPayload = LodMeshPayloadPerTranslatorPayloadKey.Add(MeshNodeContext.TranslatorPayloadKey);
+					SkeletalMeshTranslatorPayloadInterface->GetSkeletalMeshLodPayloadData(MeshNodeContext.TranslatorPayloadKey, LodMeshPayload);
+				}
+
+				//Wait until we grab all the meshes payloads
+				SkeletalMeshTranslatorPayloadInterface->WaitUntilAllSkeletalMeshPayloadCommandAreCompleted();
+
 				//Fill the lod mesh description using all combined mesh part
-				for (const FMeshNodeContext& MeshNodeContext : MeshReferences)
+				for (const UE::Interchange::Private::FMeshNodeContext& MeshNodeContext : MeshReferences)
 				{
-					TOptional<UE::Interchange::FSkeletalMeshLodPayloadData> LodMeshPayload = SkeletalMeshTranslatorPayloadInterface->GetSkeletalMeshLodPayloadData(MeshNodeContext.TranslatorPayloadKey);
+					TOptional<UE::Interchange::FSkeletalMeshLodPayloadData>& LodMeshPayload = LodMeshPayloadPerTranslatorPayloadKey.FindChecked(MeshNodeContext.TranslatorPayloadKey);
 					if (!LodMeshPayload.IsSet())
 					{
 						UE_LOG(LogInterchangeImport, Warning, TEXT("Invalid Skeletal mesh payload key [%s] SkeletalMesh asset %s"), *MeshNodeContext.TranslatorPayloadKey, *Arguments.AssetName);
@@ -715,39 +848,12 @@ UObject* UInterchangeSkeletalMeshFactory::CreateAsset(const UInterchangeSkeletal
 					FStaticMeshOperations::AppendMeshDescription(LodMeshPayload->LodMeshDescription, LodMeshDescription, AppendSettings);
 					FSkeletalMeshOperations::AppendSkinWeight(LodMeshPayload->LodMeshDescription, LodMeshDescription, SkeletalMeshAppendSettings);
 
-
-					
-					TArray<FString> BlendShapeUids;
-					MeshNodeContext.MeshNode->GetShapeDependencies(BlendShapeUids);
-					for (const FString& BlendShapeUid : BlendShapeUids)
-					{
-						if (const UInterchangeMeshNode* BlendShapeMeshNode = Cast<UInterchangeMeshNode>(Arguments.NodeContainer->GetNode(BlendShapeUid)))
-						{
-							TOptional<FString> BlendShapePayloadKey = BlendShapeMeshNode->GetPayLoadKey();
-							if (!BlendShapePayloadKey.IsSet())
-							{
-								UE_LOG(LogInterchangeImport, Warning, TEXT("Empty LOD morph target mesh reference payload when importing SkeletalMesh asset %s"), *Arguments.AssetName);
-								continue;
-							}
-							TOptional<UE::Interchange::FSkeletalMeshBlendShapePayloadData>& BlendShapeMeshPayload = BlendShapeMeshDescriptionsPerBlendShapeName.Add(BlendShapeUid);
-							BlendShapeMeshPayload = SkeletalMeshTranslatorPayloadInterface->GetSkeletalMeshBlendShapePayloadData(BlendShapePayloadKey.GetValue());
-							if (!BlendShapeMeshPayload.IsSet())
-							{
-								BlendShapeMeshDescriptionsPerBlendShapeName.Remove(BlendShapeUid);
-								UE_LOG(LogInterchangeImport, Warning, TEXT("Invalid Skeletal mesh morph target payload key [%s] SkeletalMesh asset %s"), *BlendShapePayloadKey.GetValue(), *Arguments.AssetName);
-								continue;
-							}
-							//Use the Mesh node parent bake transform
-							if (MeshNodeContext.SceneGlobalTransform.IsSet())
-							{
-								BlendShapeMeshPayload->GlobalTransform = MeshNodeContext.SceneGlobalTransform;
-							}
-							else
-							{
-								BlendShapeMeshPayload->GlobalTransform.Reset();
-							}
-						}
-					}
+					FillBlendShapeMeshDescriptionsPerBlendShapeName(MeshNodeContext
+																	, BlendShapeMeshDescriptionsPerBlendShapeName
+																	, SkeletalMeshTranslatorPayloadInterface
+																	, VertexOffset
+																	, Arguments.NodeContainer
+																	, Arguments.AssetName);
 				}
 
 				SkeletalMeshImportData = FSkeletalMeshImportData::CreateFromMeshDescription(LodMeshDescription);
diff --git a/Engine/Plugins/Runtime/Interchange/Source/Import/Public/Fbx/InterchangeFbxTranslator.h b/Engine/Plugins/Runtime/Interchange/Source/Import/Public/Fbx/InterchangeFbxTranslator.h
--- a/Engine/Plugins/Runtime/Interchange/Source/Import/Public/Fbx/InterchangeFbxTranslator.h
+++ b/Engine/Plugins/Runtime/Interchange/Source/Import/Public/Fbx/InterchangeFbxTranslator.h
@@ -78,16 +78,22 @@ public:
 	 * @param PayloadKey - The key to retrieve the a particular payload contain into the specified source data.
 	 * @return a PayloadData containing the imported data. The TOptional will not be set if there is an error.
 	 */
-	virtual TOptional<UE::Interchange::FSkeletalMeshLodPayloadData> GetSkeletalMeshLodPayloadData(const FString& PayLoadKey) const override;
+	virtual void GetSkeletalMeshLodPayloadData(const FString& PayLoadKey, TOptional<UE::Interchange::FSkeletalMeshLodPayloadData>& OptionalSkeletalMeshLodPayloadData) const override;
 
 	/**
 	 * Once the translation is done, the import process need a way to retrieve payload data.
-	 * This payload will be use by the factories to create the morph target shapes.
+	 * This payload will be use by the factories to create the morph target shapes. You have to call
+	 * WaitUntilAllSkeletalMeshPayloadCommandAreCompleted before the optional payload data is filled.
 	 *
 	 * @param PayloadKey - The key to retrieve the a particular payload contain into the specified source data.
-	 * @return a PayloadData containing the data point by the payload key. The TOptional will not be set if there is an error.
+	 * @param OptionalSkeletalMeshBlendShapePayloadData - PayloadData containing the data point by the payload key. The TOptional will not be set if there is an error.
 	 */
-	virtual TOptional<UE::Interchange::FSkeletalMeshBlendShapePayloadData> GetSkeletalMeshBlendShapePayloadData(const FString& PayLoadKey) const override;
+	virtual void GetSkeletalMeshBlendShapePayloadData(const FString& PayLoadKey, TOptional<UE::Interchange::FSkeletalMeshBlendShapePayloadData>& OptionalSkeletalMeshBlendShapePayloadData) const override;
+	
+	/**
+	 * Wait until all skeletal mesh payload commands are done
+	 */
+	virtual void WaitUntilAllSkeletalMeshPayloadCommandAreCompleted() const override;
 
 	/* IInterchangeSkeletalMeshPayloadInterface End */
 private:
diff --git a/Engine/Plugins/Runtime/Interchange/Source/Import/Public/Mesh/InterchangeSkeletalMeshPayload.h b/Engine/Plugins/Runtime/Interchange/Source/Import/Public/Mesh/InterchangeSkeletalMeshPayload.h
--- a/Engine/Plugins/Runtime/Interchange/Source/Import/Public/Mesh/InterchangeSkeletalMeshPayload.h
+++ b/Engine/Plugins/Runtime/Interchange/Source/Import/Public/Mesh/InterchangeSkeletalMeshPayload.h
@@ -23,6 +23,8 @@ namespace UE
 			//BlendShape payload is a MeshDescription containing at least the vertex position.
 			FMeshDescription LodMeshDescription;
 			TOptional<FTransform> GlobalTransform;
+			//The vertex offset of the blendshape in case we combine mesh node together
+			int32 VertexOffset;
 		};
 	}//ns Interchange
 }//ns UE
diff --git a/Engine/Plugins/Runtime/Interchange/Source/Import/Public/Mesh/InterchangeSkeletalMeshPayloadInterface.h b/Engine/Plugins/Runtime/Interchange/Source/Import/Public/Mesh/InterchangeSkeletalMeshPayloadInterface.h
--- a/Engine/Plugins/Runtime/Interchange/Source/Import/Public/Mesh/InterchangeSkeletalMeshPayloadInterface.h
+++ b/Engine/Plugins/Runtime/Interchange/Source/Import/Public/Mesh/InterchangeSkeletalMeshPayloadInterface.h
@@ -31,16 +31,22 @@ public:
 	 * @param PayloadKey - The key to retrieve the a particular payload contain into the specified source data.
 	 * @return a PayloadData containing the data point by the payload key. The TOptional will not be set if there is an error.
 	 */
-	virtual TOptional<UE::Interchange::FSkeletalMeshLodPayloadData> GetSkeletalMeshLodPayloadData(const FString& PayLoadKey) const = 0;
+	virtual void GetSkeletalMeshLodPayloadData(const FString& PayLoadKey, TOptional<UE::Interchange::FSkeletalMeshLodPayloadData>& OptionalSkeletalMeshLodPayloadData) const = 0;
 
 	/**
 	 * Once the translation is done, the import process need a way to retrieve payload data.
-	 * This payload will be use by the factories to create the morph target shapes.
+	 * This payload will be use by the factories to create the morph target shapes. You have to call
+	 * WaitUntilAllSkeletalMeshPayloadCommandAreCompleted before the optional payload data is filled.
 	 *
 	 * @param PayloadKey - The key to retrieve the a particular payload contain into the specified source data.
-	 * @return a PayloadData containing the data point by the payload key. The TOptional will not be set if there is an error.
+	 * @param OptionalSkeletalMeshBlendShapePayloadData - PayloadData containing the data point by the payload key. The TOptional will not be set if there is an error.
+	 */
+	virtual void GetSkeletalMeshBlendShapePayloadData(const FString& PayLoadKey, TOptional<UE::Interchange::FSkeletalMeshBlendShapePayloadData>& OptionalSkeletalMeshBlendShapePayloadData) const = 0;
+
+	/**
+	 * Wait until all payload commands are done
 	 */
-	virtual TOptional<UE::Interchange::FSkeletalMeshBlendShapePayloadData> GetSkeletalMeshBlendShapePayloadData(const FString& PayLoadKey) const = 0;
+	virtual void WaitUntilAllSkeletalMeshPayloadCommandAreCompleted() const = 0;
 };
 
 
diff --git a/Engine/Plugins/Runtime/Interchange/Source/Parsers/Fbx/Private/InterchangeFbxParser.cpp b/Engine/Plugins/Runtime/Interchange/Source/Parsers/Fbx/Private/InterchangeFbxParser.cpp
--- a/Engine/Plugins/Runtime/Interchange/Source/Parsers/Fbx/Private/InterchangeFbxParser.cpp
+++ b/Engine/Plugins/Runtime/Interchange/Source/Parsers/Fbx/Private/InterchangeFbxParser.cpp
@@ -35,16 +35,18 @@ namespace UE
 
 
 			ResultFilepath = ResultFolder + TEXT("/SceneDescription.itc");
-
-			Container = TStrongObjectPtr<UInterchangeBaseNodeContainer>(NewObject<UInterchangeBaseNodeContainer>(GetTransientPackage(), NAME_None));
-			if (!ensure(Container.IsValid()))
+			//Since we are not in main thread we cannot use TStrongPtr, so we will add the object to the root and remove it when we are done
+			UInterchangeBaseNodeContainer* Container = NewObject<UInterchangeBaseNodeContainer>(GetTransientPackage(), NAME_None);
+			Container->AddToRoot();
+			if (!ensure(Container != nullptr))
 			{
 				JsonLoadMessages.Add(TEXT("{\"Msg\" : {\"Type\" : \"Error\",\n\"Msg\" : \"Cannot allocate base node container to add fbx scene data\"}}"));
 				return;
 			}
-			FbxParserPrivate->FillContainerWithFbxScene(*Container.Get(), JsonLoadMessages);
+			FbxParserPrivate->FillContainerWithFbxScene(*Container, JsonLoadMessages);
 
-			Container.Get()->SaveToFile(ResultFilepath);
+			Container->SaveToFile(ResultFilepath);
+			Container->RemoveFromRoot();
 			JsonLoadMessages.Add(TEXT("{\"Msg\" : {\"Type\" : \"Log\",\n\"Msg\" : \"This is a success!\"}}"));
 		}
 
@@ -54,7 +56,9 @@ namespace UE
 
 			FString& PayloadFilepath = ResultPayloads.FindOrAdd(PayloadKey);
 			PayloadFilepath = ResultFolder + TEXT("/") + PayloadKey + TEXT(".payload");
-			if (!FbxParserPrivate->FetchPayloadData(PayloadKey, PayloadFilepath, JsonLoadMessages))
+			//Copy the map filename key because we are multithreaded and the TMap can be reallocated
+			FString PayloadFilepathCopy = PayloadFilepath;
+			if (!FbxParserPrivate->FetchPayloadData(PayloadKey, PayloadFilepathCopy, JsonLoadMessages))
 			{
 				JsonLoadMessages.Add(TEXT("{\"Msg\" : {\"Type\" : \"Error\",\n\"Msg\" : \"Cannot fetch fbx payload data.\"}}"));
 				return;
diff --git a/Engine/Plugins/Runtime/Interchange/Source/Parsers/Fbx/Public/InterchangeFbxParser.h b/Engine/Plugins/Runtime/Interchange/Source/Parsers/Fbx/Public/InterchangeFbxParser.h
--- a/Engine/Plugins/Runtime/Interchange/Source/Parsers/Fbx/Public/InterchangeFbxParser.h
+++ b/Engine/Plugins/Runtime/Interchange/Source/Parsers/Fbx/Public/InterchangeFbxParser.h
@@ -46,7 +46,6 @@ namespace UE
 			TArray<FString> GetJsonLoadMessages() const { return JsonLoadMessages; }
 
 		private:
-			TStrongObjectPtr<UInterchangeBaseNodeContainer> Container = nullptr;
 			FString ResultFilepath;
 			TArray<FString> JsonLoadMessages;
 			TMap<FString, FString> ResultPayloads;
diff --git a/Engine/Source/Programs/InterchangeWorker/Private/InterchangeWorkerImpl.cpp b/Engine/Source/Programs/InterchangeWorker/Private/InterchangeWorkerImpl.cpp
--- a/Engine/Source/Programs/InterchangeWorker/Private/InterchangeWorkerImpl.cpp
+++ b/Engine/Source/Programs/InterchangeWorker/Private/InterchangeWorkerImpl.cpp
@@ -49,7 +49,7 @@ bool FInterchangeWorkerImpl::Run()
 	bool bIsRunning = true;
 	while (bIsRunning)
 	{
-		if (TSharedPtr<ICommand> Command = CommandIO.GetNextCommand(1.0))
+		if (TSharedPtr<ICommand> Command = CommandIO.GetNextCommand(0.02))
 		{
 			switch(Command->GetType())
 			{
@@ -62,7 +62,13 @@ bool FInterchangeWorkerImpl::Run()
 					break;
 
 				case ECommandId::RunTask:
-					ProcessCommand(*StaticCast<FRunTaskCommand*>(Command.Get()));
+				{
+					FScopeLock Lock(&TFinishThreadCriticalSection);
+					//Use a randomGuid to generate thread unique name
+					int32 UniqueID = (FPlatformTime::Cycles64() & 0x00000000EFFFFFFF);
+					FString ThreadName = FString(TEXT("InterchangeWorkerCommand_")) + FString::FromInt(UniqueID);
+					IOThreads.Add(ThreadName, FThread(*ThreadName, [this, ThreadName, Command]() { ProcessCommand(Command, ThreadName); }));
+				}
 					break;
 
 				case ECommandId::Terminate:
@@ -83,10 +89,31 @@ bool FInterchangeWorkerImpl::Run()
 				UE_CLOG(!bIsRunning, LogInterchangeWorker, Error, TEXT("Worker failure: server lost"));
 			}
 		}
+		
+		//Cleanup Finish threads
+		{
+			FScopeLock Lock(&TFinishThreadCriticalSection);
+			for (const FString& ThreadName : CurrentFinishThreads)
+			{
+				if (FThread* IOThread = IOThreads.Find(ThreadName))
+				{
+					IOThread->Join();
+				}
+				IOThreads.Remove(ThreadName);
+			}
+			CurrentFinishThreads.Empty();
+		}
+
 		//Sleep 0 to avoid using too much cpu
 		FPlatformProcess::Sleep(0.0f);
 	}
 
+	//Join all thread that is not terminate
+	for(TPair<FString, FThread>&IOThread : IOThreads)
+	{
+		IOThread.Value.Join();
+	}
+
 	UE_CLOG(!bIsRunning, LogInterchangeWorker, Verbose, TEXT("Worker loop exit..."));
 	CommandIO.Disconnect(0);
 	return true;
@@ -115,8 +142,21 @@ void FInterchangeWorkerImpl::ProcessCommand(const FBackPingCommand& BackPingComm
 	PingStartCycle = 0;
 }
 
-void FInterchangeWorkerImpl::ProcessCommand(const FRunTaskCommand& RunTaskCommand)
+void FInterchangeWorkerImpl::ProcessCommand(const TSharedPtr<UE::Interchange::ICommand> Command, const FString& ThreadName)
 {
+	if (!Command.IsValid())
+	{
+		UE_LOG(LogInterchangeWorker, Error, TEXT("Process command error: The run task command is invalid"));
+		return;
+	}
+	FRunTaskCommand* RunTaskCommandPtr = StaticCast<FRunTaskCommand*>(Command.Get());
+	if (!RunTaskCommandPtr)
+	{
+		UE_LOG(LogInterchangeWorker, Error, TEXT("Process command error: The run task command is invalid"));
+		return;
+	}
+	const FRunTaskCommand& RunTaskCommand = *RunTaskCommandPtr;
+
 	const FString& JsonToProcess = RunTaskCommand.JsonDescription;
 	UE_LOG(LogInterchangeWorker, Verbose, TEXT("Process %s"), *JsonToProcess);
 
@@ -153,14 +193,20 @@ void FInterchangeWorkerImpl::ProcessCommand(const FRunTaskCommand& RunTaskComman
 	FCompletedTaskCommand CompletedTask;
 	CompletedTask.ProcessResult = ProcessResult;
 	CompletedTask.JSonMessages = JSonMessages;
+	CompletedTask.TaskIndex = RunTaskCommand.TaskIndex;
 	if (CompletedTask.ProcessResult == ETaskState::ProcessOk)
 	{
 		CompletedTask.JSonResult = JSonResult;
 	}
 
 	CommandIO.SendCommand(CompletedTask, Config::SendCommandTimeout_s);
-
 	UE_LOG(LogInterchangeWorker, Verbose, TEXT("End of Process %s"), *JsonToProcess);
+	
+	//Notify the main thread we are done with this thread
+	{
+		FScopeLock Lock(&TFinishThreadCriticalSection);
+		CurrentFinishThreads.Add(ThreadName);
+	}
 }
 
 ETaskState FInterchangeWorkerImpl::LoadFbxFile(const FJsonLoadSourceCmd& LoadSourceCommand, FString& OutJSonResult, TArray<FString>& OutJSonMessages)
diff --git a/Engine/Source/Programs/InterchangeWorker/Private/InterchangeWorkerImpl.h b/Engine/Source/Programs/InterchangeWorker/Private/InterchangeWorkerImpl.h
--- a/Engine/Source/Programs/InterchangeWorker/Private/InterchangeWorkerImpl.h
+++ b/Engine/Source/Programs/InterchangeWorker/Private/InterchangeWorkerImpl.h
@@ -1,7 +1,9 @@
 // Copyright Epic Games, Inc. All Rights Reserved.
 #pragma once
 
+#include "HAL/CriticalSection.h"
 #include "HAL/PlatformTime.h"
+#include "HAL/Thread.h"
 #include "InterchangeWorker.h"
 #include "InterchangeCommands.h"
 #include "InterchangeDispatcherNetworking.h"
@@ -20,7 +22,7 @@ private:
 	void InitiatePing();
 	void ProcessCommand(const UE::Interchange::FPingCommand& PingCommand);
 	void ProcessCommand(const UE::Interchange::FBackPingCommand& BackPingCommand);
-	void ProcessCommand(const UE::Interchange::FRunTaskCommand& TerminateCommand);
+	void ProcessCommand(const TSharedPtr<UE::Interchange::ICommand> Command, const FString& ThreadName);
 
 	UE::Interchange::ETaskState LoadFbxFile(const UE::Interchange::FJsonLoadSourceCmd& LoadSourceCommand, FString& OutJSonResult, TArray<FString>& OutJSonMessages);
 	UE::Interchange::ETaskState FetchFbxPayload(const UE::Interchange::FJsonFetchPayloadCmd& FetchPayloadCommand, FString& OutJSonResult, TArray<FString>& OutJSonMessages);
@@ -33,6 +35,9 @@ private:
 	int32 ServerPort;
 	uint64 PingStartCycle;
 	FString ResultFolder;
+	FCriticalSection TFinishThreadCriticalSection;
+	TArray<FString> CurrentFinishThreads;
+	TMap<FString, FThread > IOThreads;
 
 	UE::Interchange::FInterchangeFbxParser FbxParser;
 
diff --git a/Engine/Source/Runtime/Engine/Private/SkeletalMesh.cpp b/Engine/Source/Runtime/Engine/Private/SkeletalMesh.cpp
--- a/Engine/Source/Runtime/Engine/Private/SkeletalMesh.cpp
+++ b/Engine/Source/Runtime/Engine/Private/SkeletalMesh.cpp
@@ -3534,17 +3534,24 @@ bool USkeletalMesh::RegisterMorphTarget(UMorphTarget* MorphTarget, bool bInvalid
 		MorphTarget->BaseSkelMesh = this;
 
 		bool bRegistered = false;
-
-		for ( int32 Index = 0; Index < GetMorphTargets().Num(); ++Index )
+		TArray<UMorphTarget*>& RegisteredMorphTargets = GetMorphTargets();
+		const FName MorphTargetName = MorphTarget->GetFName();
+		if (UMorphTarget** MorphTargetPtrPtr = RegisteredMorphTargets.FindByPredicate([&MorphTargetName](UMorphTarget* RegisteredMorphTarget)->bool{ return (RegisteredMorphTarget->GetFName() == MorphTargetName); }))
 		{
-			if (GetMorphTargets()[Index]->GetFName() == MorphTarget->GetFName() )
-			{
-				UE_LOG( LogSkeletalMesh, Verbose, TEXT("RegisterMorphTarget: %s already exists, replacing"), *MorphTarget->GetName() );
-				GetMorphTargets()[Index] = MorphTarget;
-				bRegistered = true;
-				break;
-			}
+			UE_LOG(LogSkeletalMesh, Verbose, TEXT("RegisterMorphTarget: %s already exists, replacing"), *MorphTarget->GetName());
+			*MorphTargetPtrPtr = MorphTarget;
+			bRegistered = true;
 		}
+// 		for ( int32 Index = 0; Index < GetMorphTargets().Num(); ++Index )
+// 		{
+// 			if (GetMorphTargets()[Index]->GetFName() == MorphTarget->GetFName() )
+// 			{
+// 				UE_LOG( LogSkeletalMesh, Verbose, TEXT("RegisterMorphTarget: %s already exists, replacing"), *MorphTarget->GetName() );
+// 				GetMorphTargets()[Index] = MorphTarget;
+// 				bRegistered = true;
+// 				break;
+// 			}
+// 		}
 
 		if (!bRegistered)
 		{
