Use a TFuture for interchange skeletal mesh payload interface

diff --git a/Engine/Plugins/Runtime/Interchange/Source/Import/Private/Fbx/InterchangeFbxTranslator.cpp b/Engine/Plugins/Runtime/Interchange/Source/Import/Private/Fbx/InterchangeFbxTranslator.cpp
--- a/Engine/Plugins/Runtime/Interchange/Source/Import/Private/Fbx/InterchangeFbxTranslator.cpp
+++ b/Engine/Plugins/Runtime/Interchange/Source/Import/Private/Fbx/InterchangeFbxTranslator.cpp
@@ -152,16 +152,18 @@ TOptional<UE::Interchange::FStaticMeshPayloadData> UInterchangeFbxTranslator::Ge
 	return TOptional<UE::Interchange::FStaticMeshPayloadData>();
 }
 
-void UInterchangeFbxTranslator::GetSkeletalMeshLodPayloadData(const FString& PayLoadKey, TOptional<UE::Interchange::FSkeletalMeshLodPayloadData>& OptionalSkeletalMeshLodPayloadData) const
+TFuture<TOptional<UE::Interchange::FSkeletalMeshLodPayloadData>> UInterchangeFbxTranslator::GetSkeletalMeshLodPayloadData(const FString& PayLoadKey) const
 {
+	TSharedPtr<TPromise<TOptional<UE::Interchange::FSkeletalMeshLodPayloadData>>> Promise = MakeShared<TPromise<TOptional<UE::Interchange::FSkeletalMeshLodPayloadData>>>();
 	if (!Dispatcher.IsValid())
 	{
-		return;
+		Promise->SetValue(TOptional<UE::Interchange::FSkeletalMeshLodPayloadData>{});
+		return Promise->GetFuture();
 	}
 
 	//Create a json command to read the fbx file
 	FString JsonCommand = CreateFetchPayloadFbxCommand(PayLoadKey);
-	int32 TaskIndex = Dispatcher->AddTask(JsonCommand, FInterchangeDispatcherTaskCompleted::CreateLambda([this, &OptionalSkeletalMeshLodPayloadData](const int32 TaskIndex)
+	int32 TaskIndex = Dispatcher->AddTask(JsonCommand, FInterchangeDispatcherTaskCompleted::CreateLambda([this, Promise](const int32 TaskIndex)
 	{
 		UE::Interchange::ETaskState TaskState;
 		FString JsonResult;
@@ -171,6 +173,7 @@ void UInterchangeFbxTranslator::GetSkeletalMeshLodPayloadData(const FString& Pay
 		//TODO: Parse the JSonMessage and add the message to the interchange not yet develop error messaging
 		if (TaskState != UE::Interchange::ETaskState::ProcessOk)
 		{
+			Promise->SetValue(TOptional<UE::Interchange::FSkeletalMeshLodPayloadData>{});
 			return;
 		}
 		//Grab the result file and fill the BaseNodeContainer
@@ -181,10 +184,10 @@ void UInterchangeFbxTranslator::GetSkeletalMeshLodPayloadData(const FString& Pay
 		if (!ensure(FPaths::FileExists(SkeletalMeshPayloadFilename)))
 		{
 			//TODO log an error saying the payload file do not exist even if the get payload command succeed
+			Promise->SetValue(TOptional<UE::Interchange::FSkeletalMeshLodPayloadData>{});
 			return;
 		}
-		UE::Interchange::FSkeletalMeshLodPayloadData SkeletalMeshLodPayload;
-		SkeletalMeshLodPayload.LodMeshDescription.Empty();
+		
 		//All sub object should be gone with the reset
 		TArray64<uint8> Buffer;
 		FFileHelper::LoadFileToArray(Buffer, *SkeletalMeshPayloadFilename);
@@ -193,9 +196,13 @@ void UInterchangeFbxTranslator::GetSkeletalMeshLodPayloadData(const FString& Pay
 		if (FileDataSize < 1)
 		{
 			//Nothing to load from this file
+			Promise->SetValue(TOptional<UE::Interchange::FSkeletalMeshLodPayloadData>{});
 			return;
 		}
 
+		UE::Interchange::FSkeletalMeshLodPayloadData SkeletalMeshLodPayload;
+		SkeletalMeshLodPayload.LodMeshDescription.Empty();
+
 		//Buffer keep the ownership of the data, the large memory reader is use to serialize the TMap
 		FLargeMemoryReader Ar(FileData, FileDataSize);
 		SkeletalMeshLodPayload.LodMeshDescription.Serialize(Ar);
@@ -206,21 +213,25 @@ void UInterchangeFbxTranslator::GetSkeletalMeshLodPayloadData(const FString& Pay
 			//Read the bone Name to remap the influence correctly
 			Ar << SkeletalMeshLodPayload.JointNames;
 		}
-		OptionalSkeletalMeshLodPayloadData = SkeletalMeshLodPayload;
+		Promise->SetValue(MoveTemp(SkeletalMeshLodPayload));
 	}));
+
+	return Promise->GetFuture();
 }
 
-void UInterchangeFbxTranslator::GetSkeletalMeshBlendShapePayloadData(const FString& PayLoadKey, TOptional<UE::Interchange::FSkeletalMeshBlendShapePayloadData>& OptionalSkeletalMeshBlendShapePayloadData) const
+TFuture<TOptional<UE::Interchange::FSkeletalMeshBlendShapePayloadData>> UInterchangeFbxTranslator::GetSkeletalMeshBlendShapePayloadData(const FString& PayLoadKey) const
 {
+	TSharedPtr<TPromise<TOptional<UE::Interchange::FSkeletalMeshBlendShapePayloadData>>> Promise = MakeShared<TPromise<TOptional<UE::Interchange::FSkeletalMeshBlendShapePayloadData>>>();
 	if (!Dispatcher.IsValid())
 	{
-		return;
+		Promise->SetValue(TOptional<UE::Interchange::FSkeletalMeshBlendShapePayloadData>{});
+		return Promise->GetFuture();
 	}
 
 	//Create a json command to read the fbx file
 	FString JsonCommand = CreateFetchPayloadFbxCommand(PayLoadKey);
 
-	const int32 CreatedTaskIndex = Dispatcher->AddTask(JsonCommand,	FInterchangeDispatcherTaskCompleted::CreateLambda([this, &OptionalSkeletalMeshBlendShapePayloadData](const int32 TaskIndex)
+	const int32 CreatedTaskIndex = Dispatcher->AddTask(JsonCommand, FInterchangeDispatcherTaskCompleted::CreateLambda([this, Promise](const int32 TaskIndex)
 	{
 		UE::Interchange::ETaskState TaskState;
 		FString JsonResult;
@@ -231,6 +242,7 @@ void UInterchangeFbxTranslator::GetSkeletalMeshBlendShapePayloadData(const FStri
 
 		if (TaskState != UE::Interchange::ETaskState::ProcessOk)
 		{
+			Promise->SetValue(TOptional<UE::Interchange::FSkeletalMeshBlendShapePayloadData>{});
 			return;
 		}
 		//Grab the result file and fill the BaseNodeContainer
@@ -241,10 +253,9 @@ void UInterchangeFbxTranslator::GetSkeletalMeshBlendShapePayloadData(const FStri
 		if (!ensure(FPaths::FileExists(SkeletalMeshPayloadFilename)))
 		{
 			//TODO log an error saying the payload file do not exist even if the get payload command succeed
+			Promise->SetValue(TOptional<UE::Interchange::FSkeletalMeshBlendShapePayloadData>{});
 			return;
 		}
-		UE::Interchange::FSkeletalMeshBlendShapePayloadData SkeletalMeshBlendShapePayload;
-		SkeletalMeshBlendShapePayload.LodMeshDescription.Empty();
 		//All sub object should be gone with the reset
 		TArray64<uint8> Buffer;
 		FFileHelper::LoadFileToArray(Buffer, *SkeletalMeshPayloadFilename);
@@ -253,20 +264,18 @@ void UInterchangeFbxTranslator::GetSkeletalMeshBlendShapePayloadData(const FStri
 		if (FileDataSize < 1)
 		{
 			//Nothing to load from this file
+			Promise->SetValue(TOptional<UE::Interchange::FSkeletalMeshBlendShapePayloadData>{});
 			return;
 		}
-
+		UE::Interchange::FSkeletalMeshBlendShapePayloadData SkeletalMeshBlendShapePayload;
+		SkeletalMeshBlendShapePayload.LodMeshDescription.Empty();
 		//Buffer keep the ownership of the data, the large memory reader is use to serialize the TMap
 		FLargeMemoryReader Ar(FileData, FileDataSize);
 		SkeletalMeshBlendShapePayload.LodMeshDescription.Serialize(Ar);
-		OptionalSkeletalMeshBlendShapePayloadData = SkeletalMeshBlendShapePayload;
+		Promise->SetValue(MoveTemp(SkeletalMeshBlendShapePayload));
 	}));
-}
 
-void UInterchangeFbxTranslator::WaitUntilAllSkeletalMeshPayloadCommandAreCompleted() const
-{
-	//Blocking call until all tasks are executed
-	Dispatcher->WaitAllTaskToCompleteExecution();
+	return Promise->GetFuture();
 }
 
 FString UInterchangeFbxTranslator::CreateLoadFbxFileCommand(const FString& FbxFilePath) const
diff --git a/Engine/Plugins/Runtime/Interchange/Source/Import/Private/Mesh/InterchangeSkeletalMeshFactory.cpp b/Engine/Plugins/Runtime/Interchange/Source/Import/Private/Mesh/InterchangeSkeletalMeshFactory.cpp
--- a/Engine/Plugins/Runtime/Interchange/Source/Import/Private/Mesh/InterchangeSkeletalMeshFactory.cpp
+++ b/Engine/Plugins/Runtime/Interchange/Source/Import/Private/Mesh/InterchangeSkeletalMeshFactory.cpp
@@ -154,8 +154,8 @@
 			{
 				TArray<FString> BlendShapeUids;
 				MeshNodeContext.MeshNode->GetShapeDependencies(BlendShapeUids);
-				//Make sure we do not reallocate the map since its pass to the asynchronous GetSkeletalMeshBlendShapePayloadData.
-				BlendShapeMeshDescriptionsPerBlendShapeName.Reserve(BlendShapeUids.Num());
+				TMap<FString, TFuture<TOptional<UE::Interchange::FSkeletalMeshBlendShapePayloadData>>> TempBlendShapeMeshDescriptionsPerBlendShapeName;
+				TempBlendShapeMeshDescriptionsPerBlendShapeName.Reserve(BlendShapeUids.Num());
 				for (const FString& BlendShapeUid : BlendShapeUids)
 				{
 					if (const UInterchangeMeshNode* BlendShapeMeshNode = Cast<UInterchangeMeshNode>(NodeContainer->GetNode(BlendShapeUid)))
@@ -166,33 +166,32 @@
 							UE_LOG(LogInterchangeImport, Warning, TEXT("Empty LOD morph target mesh reference payload when importing SkeletalMesh asset %s"), *AssetName);
 							continue;
 						}
-						TOptional<UE::Interchange::FSkeletalMeshBlendShapePayloadData>& BlendShapeMeshPayload = BlendShapeMeshDescriptionsPerBlendShapeName.Add(BlendShapeUid);
-						SkeletalMeshTranslatorPayloadInterface->GetSkeletalMeshBlendShapePayloadData(BlendShapePayloadKey.GetValue(), BlendShapeMeshPayload);
+						const FString PayloadKey = BlendShapePayloadKey.GetValue();
+						//Add the map entry key, the translator will be call after to bulk get all the needed payload
+						TempBlendShapeMeshDescriptionsPerBlendShapeName.Add(PayloadKey, SkeletalMeshTranslatorPayloadInterface->GetSkeletalMeshBlendShapePayloadData(PayloadKey));
 					}
 				}
 
-				//Wait until all payload are loaded
-				SkeletalMeshTranslatorPayloadInterface->WaitUntilAllSkeletalMeshPayloadCommandAreCompleted();
-
 				for (const FString& BlendShapeUid : BlendShapeUids)
 				{
 					if (const UInterchangeMeshNode* BlendShapeMeshNode = Cast<UInterchangeMeshNode>(NodeContainer->GetNode(BlendShapeUid)))
 					{
+						
 						TOptional<FString> BlendShapePayloadKey = BlendShapeMeshNode->GetPayLoadKey();
 						if (!BlendShapePayloadKey.IsSet())
 						{
 							continue;
 						}
-						if (!ensure(BlendShapeMeshDescriptionsPerBlendShapeName.Contains(BlendShapeUid)))
+						const FString& BlendShapePayloadKeyString = BlendShapePayloadKey.GetValue();
+						if (!ensure(TempBlendShapeMeshDescriptionsPerBlendShapeName.Contains(BlendShapePayloadKeyString)))
 						{
 							continue;
 						}
 
-						TOptional<UE::Interchange::FSkeletalMeshBlendShapePayloadData>& BlendShapeMeshPayload = BlendShapeMeshDescriptionsPerBlendShapeName.FindChecked(BlendShapeUid);
+						TOptional<UE::Interchange::FSkeletalMeshBlendShapePayloadData> BlendShapeMeshPayload = TempBlendShapeMeshDescriptionsPerBlendShapeName.FindChecked(BlendShapePayloadKeyString).Get();
 						if (!BlendShapeMeshPayload.IsSet())
 						{
-							BlendShapeMeshDescriptionsPerBlendShapeName.Remove(BlendShapeUid);
-							UE_LOG(LogInterchangeImport, Warning, TEXT("Invalid Skeletal mesh morph target payload key [%s] SkeletalMesh asset %s"), *BlendShapePayloadKey.GetValue(), *AssetName);
+							UE_LOG(LogInterchangeImport, Warning, TEXT("Invalid Skeletal mesh morph target payload key [%s] SkeletalMesh asset %s"), *BlendShapePayloadKeyString, *AssetName);
 							continue;
 						}
 						BlendShapeMeshPayload->VertexOffset = VertexOffset;
@@ -205,6 +204,13 @@
 						{
 							BlendShapeMeshPayload->GlobalTransform.Reset();
 						}
+
+						if (!BlendShapeMeshNode->GetBlendShapeName(BlendShapeMeshPayload->BlendShapeName))
+						{
+							BlendShapeMeshPayload->BlendShapeName = BlendShapePayloadKeyString;
+						}
+						//Add the Blend shape to the blend shape map
+						BlendShapeMeshDescriptionsPerBlendShapeName.Add(BlendShapePayloadKeyString, BlendShapeMeshPayload);
 					}
 				}
 			}
@@ -261,8 +267,8 @@
 							//Executing the last morph target group, and we do not have a full group.
 							break;
 						}
-						const FString BlendShapeName(Keys[MorphTargetIndex]);
-						const TOptional<UE::Interchange::FSkeletalMeshBlendShapePayloadData>& BlendShapePayloadData = LodBlendShapeMeshDescriptions.FindChecked(BlendShapeName);
+						const FString BlendShapeKey(Keys[MorphTargetIndex]);
+						const TOptional<UE::Interchange::FSkeletalMeshBlendShapePayloadData>& BlendShapePayloadData = LodBlendShapeMeshDescriptions.FindChecked(BlendShapeKey);
 						if (!ensure(BlendShapePayloadData.IsSet()))
 						{
 							//This error was suppose to be catch in the pre parallel for loop
@@ -286,7 +292,7 @@
 
 						//Create the morph target source data
 						FString& MorphTargetName = DestinationSkeletalMeshImportData.MorphTargetNames[MorphTargetIndex];
-						MorphTargetName = BlendShapeName;
+						MorphTargetName = BlendShapePayloadData->BlendShapeName;
 						TSet<uint32>& ModifiedPoints = DestinationSkeletalMeshImportData.MorphTargetModifiedPoints[MorphTargetIndex];
 						FSkeletalMeshImportData& MorphTargetData = DestinationSkeletalMeshImportData.MorphTargets[MorphTargetIndex];
 
@@ -311,10 +317,109 @@
 						}
 						MorphTargetData.Points = CompressPoints;
 					}
-				});
+				}
+				, EParallelForFlags::BackgroundPriority);
 				return true;
 			}
 
+			void RetrieveAllSkeletalMeshPayloadsAndFillImportData(FSkeletalMeshImportData& DestinationImportData
+																  , TArray<FMeshNodeContext>& MeshReferences
+																  , TArray<SkeletalMeshImportData::FBone>& RefBonesBinary
+																  , const UInterchangeSkeletalMeshFactory::FCreateAssetParams& Arguments
+																  , const IInterchangeSkeletalMeshPayloadInterface* SkeletalMeshTranslatorPayloadInterface)
+			{
+				if (!SkeletalMeshTranslatorPayloadInterface)
+				{
+					return;
+				}
+				FMeshDescription LodMeshDescription;
+				FSkeletalMeshAttributes SkeletalMeshAttributes(LodMeshDescription);
+				SkeletalMeshAttributes.Register();
+				FStaticMeshOperations::FAppendSettings AppendSettings;
+				for (int32 ChannelIdx = 0; ChannelIdx < FStaticMeshOperations::FAppendSettings::MAX_NUM_UV_CHANNELS; ++ChannelIdx)
+				{
+					AppendSettings.bMergeUVChannels[ChannelIdx] = true;
+				}
+
+				TMap<FString, TFuture<TOptional<UE::Interchange::FSkeletalMeshLodPayloadData>>> LodMeshPayloadPerTranslatorPayloadKey;
+				LodMeshPayloadPerTranslatorPayloadKey.Reserve(MeshReferences.Num());
+
+				TMap<FString, TOptional<UE::Interchange::FSkeletalMeshBlendShapePayloadData>> BlendShapeMeshDescriptionsPerBlendShapeName;
+				int32 BlendShapeCount = 0;
+
+				for (const FMeshNodeContext& MeshNodeContext : MeshReferences)
+				{
+					//Add the payload entry key, the payload data will be fill later in bulk by the translator
+					LodMeshPayloadPerTranslatorPayloadKey.Add(MeshNodeContext.TranslatorPayloadKey, SkeletalMeshTranslatorPayloadInterface->GetSkeletalMeshLodPayloadData(MeshNodeContext.TranslatorPayloadKey));
+					//Count the blend shape dependencies so we can reserve the right amount
+					BlendShapeCount += MeshNodeContext.MeshNode ? MeshNodeContext.MeshNode->GetShapeDependeciesCount() : 0;
+				}
+				BlendShapeMeshDescriptionsPerBlendShapeName.Reserve(BlendShapeCount);
+
+				//Fill the lod mesh description using all combined mesh part
+				for (const FMeshNodeContext& MeshNodeContext : MeshReferences)
+				{
+					TOptional<UE::Interchange::FSkeletalMeshLodPayloadData> LodMeshPayload = LodMeshPayloadPerTranslatorPayloadKey.FindChecked(MeshNodeContext.TranslatorPayloadKey).Get();
+					if (!LodMeshPayload.IsSet())
+					{
+						UE_LOG(LogInterchangeImport, Warning, TEXT("Invalid Skeletal mesh payload key [%s] SkeletalMesh asset %s"), *MeshNodeContext.TranslatorPayloadKey, *Arguments.AssetName);
+						continue;
+					}
+					const int32 VertexOffset = LodMeshDescription.Vertices().Num();
+
+					FSkeletalMeshOperations::FSkeletalMeshAppendSettings SkeletalMeshAppendSettings;
+					SkeletalMeshAppendSettings.SourceVertexIDOffset = VertexOffset;
+					FElementIDRemappings ElementIDRemappings;
+					LodMeshPayload->LodMeshDescription.Compact(ElementIDRemappings);
+					//Remap the influence vertex index to point on the correct index
+					if (LodMeshPayload->JointNames.Num() > 0)
+					{
+						const int32 LocalJointCount = LodMeshPayload->JointNames.Num();
+						const int32 RefBoneCount = RefBonesBinary.Num();
+						SkeletalMeshAppendSettings.SourceRemapBoneIndex.AddZeroed(LocalJointCount);
+						for (int32 LocalJointIndex = 0; LocalJointIndex < LocalJointCount; ++LocalJointIndex)
+						{
+							SkeletalMeshAppendSettings.SourceRemapBoneIndex[LocalJointIndex] = LocalJointIndex;
+							const FString& LocalJointName = LodMeshPayload->JointNames[LocalJointIndex];
+							for (int32 RefBoneIndex = 0; RefBoneIndex < RefBoneCount; ++RefBoneIndex)
+							{
+								const SkeletalMeshImportData::FBone& Bone = RefBonesBinary[RefBoneIndex];
+								if (Bone.Name.Equals(LocalJointName))
+								{
+									SkeletalMeshAppendSettings.SourceRemapBoneIndex[LocalJointIndex] = RefBoneIndex;
+									break;
+								}
+							}
+						}
+					}
+
+					//Bake the payload, with the provide transform
+					if (MeshNodeContext.SceneGlobalTransform.IsSet())
+					{
+						AppendSettings.MeshTransform = MeshNodeContext.SceneGlobalTransform;
+					}
+					else
+					{
+						AppendSettings.MeshTransform.Reset();
+					}
+					FStaticMeshOperations::AppendMeshDescription(LodMeshPayload->LodMeshDescription, LodMeshDescription, AppendSettings);
+					FSkeletalMeshOperations::AppendSkinWeight(LodMeshPayload->LodMeshDescription, LodMeshDescription, SkeletalMeshAppendSettings);
+
+					FillBlendShapeMeshDescriptionsPerBlendShapeName(MeshNodeContext
+																	, BlendShapeMeshDescriptionsPerBlendShapeName
+																	, SkeletalMeshTranslatorPayloadInterface
+																	, VertexOffset
+																	, Arguments.NodeContainer
+																	, Arguments.AssetName);
+				}
+
+				DestinationImportData = FSkeletalMeshImportData::CreateFromMeshDescription(LodMeshDescription);
+				DestinationImportData.RefBonesBinary = RefBonesBinary;
+
+				//Copy all the lod blend shapes data to the DestinationImportData.
+				CopyBlendShapesMeshDescriptionToSkeletalMeshImportData(BlendShapeMeshDescriptionsPerBlendShapeName, DestinationImportData);
+			}
+
 			/**
 			 * Fill the Materials array using the raw skeletalmesh geometry data (using material imported name)
 			 * Find the material from the dependencies of the skeletalmesh before searching in all package.
@@ -776,91 +881,13 @@
 
 				//Add the lod mesh data to the skeletalmesh
 				FSkeletalMeshImportData SkeletalMeshImportData;
-				FMeshDescription LodMeshDescription;
-				FSkeletalMeshAttributes SkeletalMeshAttributes(LodMeshDescription);
-				SkeletalMeshAttributes.Register();
-				FStaticMeshOperations::FAppendSettings AppendSettings;
-				for (int32 ChannelIdx = 0; ChannelIdx < FStaticMeshOperations::FAppendSettings::MAX_NUM_UV_CHANNELS; ++ChannelIdx)
-				{
-					AppendSettings.bMergeUVChannels[ChannelIdx] = true;
-				}
 				
-				//Blend shape MeshDescriptions
-				TMap<FString, TOptional<UE::Interchange::FSkeletalMeshBlendShapePayloadData>> BlendShapeMeshDescriptionsPerBlendShapeName;
+				//Get all meshes and blend shapes payload and fill the SkeletalMeshImportData structure
+				UE::Interchange::Private::RetrieveAllSkeletalMeshPayloadsAndFillImportData(SkeletalMeshImportData
+																						   , MeshReferences
+																						   , RefBonesBinary
+																						   , Arguments
+																						   , SkeletalMeshTranslatorPayloadInterface);
 
-				TMap<FString, TOptional<UE::Interchange::FSkeletalMeshLodPayloadData>> LodMeshPayloadPerTranslatorPayloadKey;
-				LodMeshPayloadPerTranslatorPayloadKey.Reserve(MeshReferences.Num());
-				for (const UE::Interchange::Private::FMeshNodeContext& MeshNodeContext : MeshReferences)
-				{
-					TOptional<UE::Interchange::FSkeletalMeshLodPayloadData>& LodMeshPayload = LodMeshPayloadPerTranslatorPayloadKey.Add(MeshNodeContext.TranslatorPayloadKey);
-					SkeletalMeshTranslatorPayloadInterface->GetSkeletalMeshLodPayloadData(MeshNodeContext.TranslatorPayloadKey, LodMeshPayload);
-				}
-
-				//Wait until we grab all the meshes payloads
-				SkeletalMeshTranslatorPayloadInterface->WaitUntilAllSkeletalMeshPayloadCommandAreCompleted();
-
-				//Fill the lod mesh description using all combined mesh part
-				for (const UE::Interchange::Private::FMeshNodeContext& MeshNodeContext : MeshReferences)
-				{
-					TOptional<UE::Interchange::FSkeletalMeshLodPayloadData>& LodMeshPayload = LodMeshPayloadPerTranslatorPayloadKey.FindChecked(MeshNodeContext.TranslatorPayloadKey);
-					if (!LodMeshPayload.IsSet())
-					{
-						UE_LOG(LogInterchangeImport, Warning, TEXT("Invalid Skeletal mesh payload key [%s] SkeletalMesh asset %s"), *MeshNodeContext.TranslatorPayloadKey, *Arguments.AssetName);
-						continue;
-					}
-					const int32 VertexOffset = LodMeshDescription.Vertices().Num();
-					
-					FSkeletalMeshOperations::FSkeletalMeshAppendSettings SkeletalMeshAppendSettings;
-					SkeletalMeshAppendSettings.SourceVertexIDOffset = VertexOffset;
-					FElementIDRemappings ElementIDRemappings;
-					LodMeshPayload->LodMeshDescription.Compact(ElementIDRemappings);
-					//Remap the influence vertex index to point on the correct index
-					if (LodMeshPayload->JointNames.Num() > 0)
-					{
-						const int32 LocalJointCount = LodMeshPayload->JointNames.Num();
-						const int32 RefBoneCount = RefBonesBinary.Num();
-						SkeletalMeshAppendSettings.SourceRemapBoneIndex.AddZeroed(LocalJointCount);
-						for(int32 LocalJointIndex = 0; LocalJointIndex < LocalJointCount; ++LocalJointIndex)
-						{
-							SkeletalMeshAppendSettings.SourceRemapBoneIndex[LocalJointIndex] = LocalJointIndex;
-							const FString& LocalJointName = LodMeshPayload->JointNames[LocalJointIndex];
-							for(int32 RefBoneIndex = 0; RefBoneIndex < RefBoneCount; ++RefBoneIndex)
-							{
-								const SkeletalMeshImportData::FBone& Bone = RefBonesBinary[RefBoneIndex];
-								if(Bone.Name.Equals(LocalJointName))
-								{
-									SkeletalMeshAppendSettings.SourceRemapBoneIndex[LocalJointIndex] = RefBoneIndex;
-									break;
-								}
-							}
-						}
-					}
-
-					//Bake the payload, with the provide transform
-					if (MeshNodeContext.SceneGlobalTransform.IsSet())
-					{
-						AppendSettings.MeshTransform = MeshNodeContext.SceneGlobalTransform;
-					}
-					else
-					{
-						AppendSettings.MeshTransform.Reset();
-					}
-					FStaticMeshOperations::AppendMeshDescription(LodMeshPayload->LodMeshDescription, LodMeshDescription, AppendSettings);
-					FSkeletalMeshOperations::AppendSkinWeight(LodMeshPayload->LodMeshDescription, LodMeshDescription, SkeletalMeshAppendSettings);
-
-					FillBlendShapeMeshDescriptionsPerBlendShapeName(MeshNodeContext
-																	, BlendShapeMeshDescriptionsPerBlendShapeName
-																	, SkeletalMeshTranslatorPayloadInterface
-																	, VertexOffset
-																	, Arguments.NodeContainer
-																	, Arguments.AssetName);
-				}
-
-				SkeletalMeshImportData = FSkeletalMeshImportData::CreateFromMeshDescription(LodMeshDescription);
-				SkeletalMeshImportData.RefBonesBinary = RefBonesBinary;
-
-				//Copy all the lod blend shapes data to the SkeletalMeshImportData.
-				UE::Interchange::Private::CopyBlendShapesMeshDescriptionToSkeletalMeshImportData(BlendShapeMeshDescriptionsPerBlendShapeName, SkeletalMeshImportData);
-
 				ensure(ImportedResource->LODModels.Add(new FSkeletalMeshLODModel()) == CurrentLodIndex);
 				FSkeletalMeshLODModel& LODModel = ImportedResource->LODModels[CurrentLodIndex];
diff --git a/Engine/Plugins/Runtime/Interchange/Source/Import/Public/Fbx/InterchangeFbxTranslator.h b/Engine/Plugins/Runtime/Interchange/Source/Import/Public/Fbx/InterchangeFbxTranslator.h
--- a/Engine/Plugins/Runtime/Interchange/Source/Import/Public/Fbx/InterchangeFbxTranslator.h
+++ b/Engine/Plugins/Runtime/Interchange/Source/Import/Public/Fbx/InterchangeFbxTranslator.h
@@ -70,34 +70,11 @@ public:
 	//////////////////////////////////////////////////////////////////////////
 	/* IInterchangeSkeletalMeshPayloadInterface Begin */
 
-	/**
-	 * Once the translation is done, the import process need a way to retrieve payload data.
-	 * This payload will be use by the factories to create the asset.
-	 *
-	 * @param SourceData - The source data containing the data to translate
-	 * @param PayloadKey - The key to retrieve the a particular payload contain into the specified source data.
-	 * @return a PayloadData containing the imported data. The TOptional will not be set if there is an error.
-	 */
-	virtual void GetSkeletalMeshLodPayloadData(const FString& PayLoadKey, TOptional<UE::Interchange::FSkeletalMeshLodPayloadData>& OptionalSkeletalMeshLodPayloadData) const override;
-
-	/**
-	 * Once the translation is done, the import process need a way to retrieve payload data.
-	 * This payload will be use by the factories to create the morph target shapes. You have to call
-	 * WaitUntilAllSkeletalMeshPayloadCommandAreCompleted before the optional payload data is filled.
-	 *
-	 * @param PayloadKey - The key to retrieve the a particular payload contain into the specified source data.
-	 * @param OptionalSkeletalMeshBlendShapePayloadData - PayloadData containing the data point by the payload key. The TOptional will not be set if there is an error.
-	 */
-	virtual void GetSkeletalMeshBlendShapePayloadData(const FString& PayLoadKey, TOptional<UE::Interchange::FSkeletalMeshBlendShapePayloadData>& OptionalSkeletalMeshBlendShapePayloadData) const override;
+	virtual TFuture<TOptional<UE::Interchange::FSkeletalMeshLodPayloadData>> GetSkeletalMeshLodPayloadData(const FString& PayLoadKey) const override;
+	virtual TFuture<TOptional<UE::Interchange::FSkeletalMeshBlendShapePayloadData>> GetSkeletalMeshBlendShapePayloadData(const FString& PayLoadKey) const override;
 	
-	/**
-	 * Wait until all skeletal mesh payload commands are done
-	 */
-	virtual void WaitUntilAllSkeletalMeshPayloadCommandAreCompleted() const override;
-
 	/* IInterchangeSkeletalMeshPayloadInterface End */
 private:
-
 	FString CreateLoadFbxFileCommand(const FString& FbxFilePath) const;
 
 	FString CreateFetchPayloadFbxCommand(const FString& FbxPayloadKey) const;
diff --git a/Engine/Plugins/Runtime/Interchange/Source/Import/Public/Mesh/InterchangeSkeletalMeshPayload.h b/Engine/Plugins/Runtime/Interchange/Source/Import/Public/Mesh/InterchangeSkeletalMeshPayload.h
--- a/Engine/Plugins/Runtime/Interchange/Source/Import/Public/Mesh/InterchangeSkeletalMeshPayload.h
+++ b/Engine/Plugins/Runtime/Interchange/Source/Import/Public/Mesh/InterchangeSkeletalMeshPayload.h
@@ -25,6 +25,8 @@ namespace UE
 			TOptional<FTransform> GlobalTransform;
 			//The vertex offset of the blendshape in case we combine mesh node together
 			int32 VertexOffset;
+			//The name of the blendshape
+			FString BlendShapeName;
 		};
 	}//ns Interchange
 }//ns UE
diff --git a/Engine/Plugins/Runtime/Interchange/Source/Import/Public/Mesh/InterchangeSkeletalMeshPayloadInterface.h b/Engine/Plugins/Runtime/Interchange/Source/Import/Public/Mesh/InterchangeSkeletalMeshPayloadInterface.h
--- a/Engine/Plugins/Runtime/Interchange/Source/Import/Public/Mesh/InterchangeSkeletalMeshPayloadInterface.h
+++ b/Engine/Plugins/Runtime/Interchange/Source/Import/Public/Mesh/InterchangeSkeletalMeshPayloadInterface.h
@@ -25,28 +25,22 @@ class INTERCHANGEIMPORT_API IInterchangeSkeletalMeshPayloadInterface
 public:
 
 	/**
-	 * Once the translation is done, the import process need a way to retrieve payload data.
-	 * This payload will be use by the factories to create the asset.
+	 * Get a skeletal mesh payload data for the specified payload key
 	 *
-	 * @param PayloadKey - The key to retrieve the a particular payload contain into the specified source data.
-	 * @return a PayloadData containing the data point by the payload key. The TOptional will not be set if there is an error.
+	 * @param PayLoadKey - The key to retrieve the a particular payload contain into the specified source data.
+	 * @return - The resulting PayloadData as a TFuture point by the PayloadKey. The TOptional will not be set if there is an error retrieving the payload.
+	 * 
 	 */
-	virtual void GetSkeletalMeshLodPayloadData(const FString& PayLoadKey, TOptional<UE::Interchange::FSkeletalMeshLodPayloadData>& OptionalSkeletalMeshLodPayloadData) const = 0;
+	virtual TFuture<TOptional<UE::Interchange::FSkeletalMeshLodPayloadData>> GetSkeletalMeshLodPayloadData(const FString& PayLoadKey) const = 0;
 
 	/**
-	 * Once the translation is done, the import process need a way to retrieve payload data.
-	 * This payload will be use by the factories to create the morph target shapes. You have to call
-	 * WaitUntilAllSkeletalMeshPayloadCommandAreCompleted before the optional payload data is filled.
+	 * Get one skeletal mesh blend shape payload data for the specified key
 	 *
-	 * @param PayloadKey - The key to retrieve the a particular payload contain into the specified source data.
-	 * @param OptionalSkeletalMeshBlendShapePayloadData - PayloadData containing the data point by the payload key. The TOptional will not be set if there is an error.
+	 * @param PayLoadKey - The key to retrieve the a particular payload contain into the specified source data.
+	 * @return - The resulting PayloadData as a TFuture point by the payload key. The TOptional will not be set if there is an error retrieving the payload.
+	 * 
 	 */
-	virtual void GetSkeletalMeshBlendShapePayloadData(const FString& PayLoadKey, TOptional<UE::Interchange::FSkeletalMeshBlendShapePayloadData>& OptionalSkeletalMeshBlendShapePayloadData) const = 0;
-
-	/**
-	 * Wait until all payload commands are done
-	 */
-	virtual void WaitUntilAllSkeletalMeshPayloadCommandAreCompleted() const = 0;
+	virtual TFuture<TOptional<UE::Interchange::FSkeletalMeshBlendShapePayloadData>> GetSkeletalMeshBlendShapePayloadData(const FString& PayLoadKey) const = 0;
 };
 
 
diff --git a/Engine/Source/Runtime/Engine/Private/SkeletalMesh.cpp b/Engine/Source/Runtime/Engine/Private/SkeletalMesh.cpp
--- a/Engine/Source/Runtime/Engine/Private/SkeletalMesh.cpp
+++ b/Engine/Source/Runtime/Engine/Private/SkeletalMesh.cpp
@@ -3534,28 +3534,22 @@ bool USkeletalMesh::RegisterMorphTarget(UMorphTarget* MorphTarget, bool bInvalid
 		MorphTarget->BaseSkelMesh = this;
 
 		bool bRegistered = false;
-		TArray<UMorphTarget*>& RegisteredMorphTargets = GetMorphTargets();
 		const FName MorphTargetName = MorphTarget->GetFName();
-		if (UMorphTarget** MorphTargetPtrPtr = RegisteredMorphTargets.FindByPredicate([&MorphTargetName](UMorphTarget* RegisteredMorphTarget)->bool{ return (RegisteredMorphTarget->GetFName() == MorphTargetName); }))
+		TArray<UMorphTarget*>& RegisteredMorphTargets = GetMorphTargets();
+		for ( int32 Index = 0; Index < RegisteredMorphTargets.Num(); ++Index )
 		{
-			UE_LOG(LogSkeletalMesh, Verbose, TEXT("RegisterMorphTarget: %s already exists, replacing"), *MorphTarget->GetName());
-			*MorphTargetPtrPtr = MorphTarget;
-			bRegistered = true;
+			if (RegisteredMorphTargets[Index]->GetFName() == MorphTargetName )
+			{
+				UE_LOG( LogSkeletalMesh, Verbose, TEXT("RegisterMorphTarget: %s already exists, replacing"), *MorphTarget->GetName() );
+				RegisteredMorphTargets[Index] = MorphTarget;
+				bRegistered = true;
+				break;
+			}
 		}
-// 		for ( int32 Index = 0; Index < GetMorphTargets().Num(); ++Index )
-// 		{
-// 			if (GetMorphTargets()[Index]->GetFName() == MorphTarget->GetFName() )
-// 			{
-// 				UE_LOG( LogSkeletalMesh, Verbose, TEXT("RegisterMorphTarget: %s already exists, replacing"), *MorphTarget->GetName() );
-// 				GetMorphTargets()[Index] = MorphTarget;
-// 				bRegistered = true;
-// 				break;
-// 			}
-// 		}
 
 		if (!bRegistered)
 		{
-			GetMorphTargets().Add( MorphTarget );
+			RegisteredMorphTargets.Add( MorphTarget );
 			bRegistered = true;
 		}
 
