Make sure we do not get stuck on a payload promise from the out of process interchange worker.

diff --git a/Engine/Plugins/Runtime/Interchange/Source/Dispatcher/Private/InterchangeCommands.cpp b/Engine/Plugins/Runtime/Interchange/Source/Dispatcher/Private/InterchangeCommands.cpp
--- a/Engine/Plugins/Runtime/Interchange/Source/Dispatcher/Private/InterchangeCommands.cpp
+++ b/Engine/Plugins/Runtime/Interchange/Source/Dispatcher/Private/InterchangeCommands.cpp
@@ -18,6 +18,8 @@ namespace UE
 				case ECommandId::BackPing: return MakeShared<FBackPingCommand>();
 				case ECommandId::RunTask: return MakeShared<FRunTaskCommand>();
 				case ECommandId::NotifyEndTask: return MakeShared<FCompletedTaskCommand>();
+				case ECommandId::QueryTaskProgress: return MakeShared<FQueryTaskProgressCommand>();
+				case ECommandId::CompletedQueryTaskProgress: return MakeShared<FCompletedQueryTaskProgressCommand>();
 				case ECommandId::Terminate: return MakeShared<FTerminateCommand>();
 			}
 			return nullptr;
@@ -57,5 +59,28 @@ namespace UE
 			Ar << JSonResult;
 			Ar << TaskIndex;
 		}
+
+		FQueryTaskProgressCommand::FQueryTaskProgressCommand(const TArray<int32>& Tasks)
+		{
+			TaskIndexes = Tasks;
+		}
+		
+		void FQueryTaskProgressCommand::SerializeImpl(FArchive& Ar)
+		{
+			Ar << TaskIndexes;
+		}
+
+		FArchive& operator<<(FArchive& Ar, FCompletedQueryTaskProgressCommand::FTaskProgressData& A)
+		{
+			Ar << A.TaskIndex;
+			Ar << A.TaskState;
+			Ar << A.TaskProgress;
+			return Ar;
+		}
+
+		void FCompletedQueryTaskProgressCommand::SerializeImpl(FArchive& Ar)
+		{
+			Ar << TaskStates;
+		}
 	}//ns Interchange
 }//ns UE
diff --git a/Engine/Plugins/Runtime/Interchange/Source/Dispatcher/Private/InterchangeDispatcher.cpp b/Engine/Plugins/Runtime/Interchange/Source/Dispatcher/Private/InterchangeDispatcher.cpp
--- a/Engine/Plugins/Runtime/Interchange/Source/Dispatcher/Private/InterchangeDispatcher.cpp
+++ b/Engine/Plugins/Runtime/Interchange/Source/Dispatcher/Private/InterchangeDispatcher.cpp
@@ -4,6 +4,7 @@
 
 #include "HAL/FileManager.h"
 #include "HAL/PlatformProcess.h"
+#include "HAL/PlatformTime.h"
 #include "InterchangeDispatcherConfig.h"
 #include "InterchangeDispatcherLog.h"
 #include "InterchangeDispatcherTask.h"
@@ -24,6 +25,10 @@ namespace UE
 
 		int32 FInterchangeDispatcher::AddTask(const FString& JsonDescription)
 		{
+			if (!WorkerHandler->IsAlive())
+			{
+				return INDEX_NONE;
+			}
 			FScopeLock Lock(&TaskPoolCriticalSection);
 			int32 TaskIndex = TaskPool.Emplace(JsonDescription);
 			TaskPool[TaskIndex].Index = TaskIndex;
@@ -32,7 +37,10 @@ namespace UE
 		int32 FInterchangeDispatcher::AddTask(const FString& JsonDescription, FInterchangeDispatcherTaskCompleted TaskCompledDelegate)
 		{
 			int32 TaskIndex = AddTask(JsonDescription);
-			TaskPool[TaskIndex].OnTaskCompleted = TaskCompledDelegate;
+			if (TaskIndex != INDEX_NONE)
+			{
+				TaskPool[TaskIndex].OnTaskCompleted = TaskCompledDelegate;
+			}
 			return TaskIndex;
 		}
 
@@ -51,6 +59,7 @@ namespace UE
 			}
 
 			TaskPool[NextTaskIndex].State = ETaskState::Running;
+			TaskPool[NextTaskIndex].RunningStateStartTime = FPlatformTime::Seconds();
 			return TaskPool[NextTaskIndex++];
 		}
 
@@ -66,6 +75,11 @@ namespace UE
 				}
 
 				FTask& Task = TaskPool[TaskIndex];
+				if (Task.State == ETaskState::ProcessOk || Task.State == ETaskState::ProcessFailed)
+				{
+					//Task was already processed, we cannot set its state after it was process
+					return;
+				}
 				Task.State = TaskState;
 				Task.JsonResult = JsonResult;
 				Task.JsonMessages = JSonMessages;
@@ -90,6 +104,20 @@ namespace UE
 			UE_CLOG(TaskState == ETaskState::ProcessFailed, LogInterchangeDispatcher, Error, TEXT("Json processing failure: %s"), *JsonDescription);
 		}
 
+		void FInterchangeDispatcher::GetTaskState(int32 TaskIndex, ETaskState& TaskState, double& TaskRunningStateStartTime)
+		{
+			FScopeLock Lock(&TaskPoolCriticalSection);
+
+			if (!ensure(TaskPool.IsValidIndex(TaskIndex)))
+			{
+				return;
+			}
+
+			FTask& Task = TaskPool[TaskIndex];
+			TaskState = Task.State;
+			TaskRunningStateStartTime = Task.RunningStateStartTime;
+		}
+
 		void FInterchangeDispatcher::GetTaskState(int32 TaskIndex, ETaskState& TaskState, FString& JsonResult, TArray<FString>& JSonMessages)
 		{
 			FScopeLock Lock(&TaskPoolCriticalSection);
@@ -124,6 +152,7 @@ namespace UE
 					WorkerHandler->Stop();
 				}
 			}
+			EmptyQueueTasks();
 		}
 
 		void FInterchangeDispatcher::TerminateProcess()
@@ -181,6 +210,10 @@ namespace UE
 		void FInterchangeDispatcher::SpawnHandler()
 		{
 			WorkerHandler = MakeUnique<FInterchangeWorkerHandler>(*this, ResultFolder);
+			WorkerHandler->OnWorkerHandlerExitLoop.AddLambda([this]()
+			{
+				EmptyQueueTasks();
+			});
 		}
 
 		bool FInterchangeDispatcher::IsHandlerAlive()
@@ -194,5 +227,18 @@ namespace UE
 			WorkerHandler.Reset();
 		}
 
+		void FInterchangeDispatcher::EmptyQueueTasks()
+		{
+			//Make sure all queue tasks are completed to process fail,
+			//This ensure any wait on completed delegate like promise of a future is unblock.
+			TOptional<FTask> NextTask = GetNextTask();
+			while (NextTask.IsSet())
+			{
+				TArray<FString> GarbageMessages;
+				SetTaskState(NextTask->Index, ETaskState::ProcessFailed, FString(), GarbageMessages);
+				NextTask = GetNextTask();
+			}
+		}
+
 	} //ns Interchange
 }//ns UE
\ No newline at end of file
diff --git a/Engine/Plugins/Runtime/Interchange/Source/Dispatcher/Private/InterchangeWorkerHandler.cpp b/Engine/Plugins/Runtime/Interchange/Source/Dispatcher/Private/InterchangeWorkerHandler.cpp
--- a/Engine/Plugins/Runtime/Interchange/Source/Dispatcher/Private/InterchangeWorkerHandler.cpp
+++ b/Engine/Plugins/Runtime/Interchange/Source/Dispatcher/Private/InterchangeWorkerHandler.cpp
@@ -106,6 +106,7 @@ namespace UE
 			int32 UniqueID = (FPlatformTime::Cycles64() & 0x00000000EFFFFFFF);
 			ThreadName = FString(TEXT("InterchangeWorkerHandler_")) + FString::FromInt(UniqueID);
 			IOThread = FThread(*ThreadName, [this]() { Run(); });
+			LastProgressMessageTime = FPlatformTime::Seconds();
 		}
 
 		FInterchangeWorkerHandler::~FInterchangeWorkerHandler()
@@ -164,6 +165,46 @@ namespace UE
 				WorkerState = EWorkerState::Closing;
 				ErrorState = EWorkerErrorState::WorkerProcess_Lost;
 			}
+			//Send periodic query progress for task that take long time to resolve
+			const double MinimumProgressTime = 5.0;
+			const double ProgressTickTime = 5.0;
+			const double CurrentTime = FPlatformTime::Seconds();
+			if(CurrentTime - LastProgressMessageTime > ProgressTickTime)
+			{
+				LastProgressMessageTime = CurrentTime;
+				FScopeLock Lock(&CurrentTasksLock);
+				TArray<int32> TasksToQuery;
+				//Look all current tasks
+				for (const int32 TaskIndex : CurrentTasks)
+				{
+					UE::Interchange::ETaskState TaskState;
+					double TaskRunningStateStartTime;
+					Dispatcher.GetTaskState(TaskIndex, TaskState, TaskRunningStateStartTime);
+					
+					if (TaskState == ETaskState::Running)
+					{
+						double RunningTime = FPlatformTime::Seconds() - TaskRunningStateStartTime;
+						//Ask only for old task that run longer then the minimum progress time
+						if (RunningTime > MinimumProgressTime)
+						{
+							TasksToQuery.Add(TaskIndex);
+						}
+					}
+				}
+				//Send the query command
+				FQueryTaskProgressCommand QueryProgressCommand(TasksToQuery);
+
+				if (CommandIO.SendCommand(QueryProgressCommand, Config::SendCommandTimeout_s))
+				{
+					UE_LOG(LogInterchangeDispatcher, Verbose, TEXT("Query progress command sent"));
+				}
+				else
+				{
+					// Cannot send a command to the worker, let set the Closing state
+					WorkerState = EWorkerState::Closing;
+					ErrorState = EWorkerErrorState::ConnectionLost_SendFailed;
+				}
+			}
 		}
 
 		void FInterchangeWorkerHandler::Run()
@@ -174,9 +215,20 @@ namespace UE
 			UE_CLOG(ErrorState != EWorkerErrorState::Ok, LogInterchangeDispatcher, Warning, TEXT("Handler ended with error: %s"), EWorkerErrorStateAsString(ErrorState));
 		}
 
-		void FInterchangeWorkerHandler::RunInternal()
+		void FInterchangeWorkerHandler::KillAllCurrentTasks()
 		{
+			FScopeLock Lock(&CurrentTasksLock);
+			//Fail all current tasks
+			for (const int32 TaskIndex : CurrentTasks)
+			{
+				TArray<FString> GarbageMessages;
+				Dispatcher.SetTaskState(TaskIndex, ETaskState::ProcessFailed, FString(), GarbageMessages);
+			}
+			CurrentTasks.Empty();
+		}
 
+		void FInterchangeWorkerHandler::RunInternal()
+		{
 			auto StartNewTask = [this]()->bool
 			{
 				// Fetch a new task
@@ -185,7 +237,10 @@ namespace UE
 				if (CurrentTask.IsSet())
 				{
 					FRunTaskCommand NewTask(CurrentTask.GetValue());
-					CurrentTasks.Add(NewTask.TaskIndex);
+					{
+						FScopeLock Lock(&CurrentTasksLock);
+						CurrentTasks.Add(NewTask.TaskIndex);
+					}
 
 					if (CommandIO.SendCommand(NewTask, Config::SendCommandTimeout_s))
 					{
@@ -196,7 +251,7 @@ namespace UE
 					{
 						// Signal that the Task was not processed
 						TArray<FString> GarbageMessages;
-						Dispatcher.SetTaskState(CurrentTask->Index, ETaskState::UnTreated, FString(), GarbageMessages);
+						Dispatcher.SetTaskState(CurrentTask->Index, ETaskState::ProcessFailed, FString(), GarbageMessages);
 						WorkerState = EWorkerState::Closing;
 						ErrorState = EWorkerErrorState::ConnectionLost_SendFailed;
 					}
@@ -246,32 +301,29 @@ namespace UE
 						ValidateConnection();
 						if (ErrorState == EWorkerErrorState::WorkerProcess_Lost)
 						{
-							for (const int32 TaskIndex : CurrentTasks)
-							{
-								TArray<FString> GarbageMessages;
-								Dispatcher.SetTaskState(TaskIndex, ETaskState::ProcessFailed, FString(), GarbageMessages);
-							}
-							CurrentTasks.Empty();
 							WorkerState = EWorkerState::Closing;
 						}
-						//Start a task
-						if (StartNewTask())
+						else
 						{
-							WorkerState = EWorkerState::Processing;
-						}
+							//Start a task
+							if (StartNewTask())
+							{
+								WorkerState = EWorkerState::Processing;
+							}
 
-						// consume task
-						if (TSharedPtr<ICommand> Command = CommandIO.GetNextCommand(Config::IdleLoopDelay))
-						{
-							//Make the processing asynchronous
-							TGraphTask<UE::Interchange::Dispatcher::FTaskProcessCommand>::CreateTask().ConstructAndDispatchWhenReady(this, Command);
-						}
+							// consume task
+							if (TSharedPtr<ICommand> Command = CommandIO.GetNextCommand(Config::IdleLoopDelay))
+							{
+								//Make the processing asynchronous
+								TGraphTask<UE::Interchange::Dispatcher::FTaskProcessCommand>::CreateTask().ConstructAndDispatchWhenReady(this, Command);
+							}
 
-						//Need to terminate?
-						if (bShouldTerminate)
-						{
-							UE_LOG(LogInterchangeDispatcher, Verbose, TEXT("Exit loop gracefully"));
-							WorkerState = EWorkerState::Closing;
+							//Need to terminate?
+							if (bShouldTerminate)
+							{
+								UE_LOG(LogInterchangeDispatcher, Verbose, TEXT("Exit loop gracefully"));
+								WorkerState = EWorkerState::Closing;
+							}
 						}
 						break;
 					}
@@ -313,6 +365,9 @@ namespace UE
 							ProcessCommand(*Command);
 						}
 
+						//Make sure all running task are completed with error
+						KillAllCurrentTasks();
+
 						WorkerState = EWorkerState::Terminated;
 						break;
 					}
@@ -323,6 +378,8 @@ namespace UE
 					}
 				}
 			}
+			//Notify the dispatcher the worker handler is done so it can terminate queue tasks
+			OnWorkerHandlerExitLoop.Broadcast();
 		}
 
 		void FInterchangeWorkerHandler::Stop()
@@ -356,6 +413,10 @@ namespace UE
 					ProcessCommand(StaticCast<FCompletedTaskCommand&>(Command));
 					break;
 
+				case ECommandId::CompletedQueryTaskProgress:
+					ProcessCommand(StaticCast<FCompletedQueryTaskProgressCommand&>(Command));
+					break;
+
 				default:
 					break;
 			}
@@ -369,9 +430,32 @@ namespace UE
 
 		void FInterchangeWorkerHandler::ProcessCommand(FCompletedTaskCommand& CompletedTaskCommand)
 		{
-
+			{
+				FScopeLock Lock(&CurrentTasksLock);
+				CurrentTasks.Remove(CompletedTaskCommand.TaskIndex);
+			}
 			Dispatcher.SetTaskState(CompletedTaskCommand.TaskIndex, CompletedTaskCommand.ProcessResult, CompletedTaskCommand.JSonResult, CompletedTaskCommand.JSonMessages);
-			CurrentTasks.Remove(CompletedTaskCommand.TaskIndex);
+		}
+
+		void FInterchangeWorkerHandler::ProcessCommand(FCompletedQueryTaskProgressCommand& CompletedQueryTaskProgressCommand)
+		{
+			for (const FCompletedQueryTaskProgressCommand::FTaskProgressData& TaskData : CompletedQueryTaskProgressCommand.TaskStates)
+			{
+				if (TaskData.TaskState != ETaskState::Running)
+				{
+					ETaskState CurrentTaskState;
+					double TaskRunningStateStartTime;
+					Dispatcher.GetTaskState(TaskData.TaskIndex, CurrentTaskState, TaskRunningStateStartTime);
+					if (CurrentTaskState == ETaskState::Running)
+					{
+						//This is a stale task, we must complete it failed. This can happen if a completed task message is skip (lost, not process, etc...)
+						//We must ensure we complete the task with a failed so if someone wait for it it will be release.
+						FString EmptyGarbage1;
+						TArray<FString> EmptyGarbage2;
+						Dispatcher.SetTaskState(TaskData.TaskIndex, ETaskState::ProcessFailed, EmptyGarbage1, EmptyGarbage2);
+					}
+				}
+			}
 		}
 
 		const TCHAR* FInterchangeWorkerHandler::EWorkerErrorStateAsString(EWorkerErrorState Error)
diff --git a/Engine/Plugins/Runtime/Interchange/Source/Dispatcher/Public/InterchangeCommands.h b/Engine/Plugins/Runtime/Interchange/Source/Dispatcher/Public/InterchangeCommands.h
--- a/Engine/Plugins/Runtime/Interchange/Source/Dispatcher/Public/InterchangeCommands.h
+++ b/Engine/Plugins/Runtime/Interchange/Source/Dispatcher/Public/InterchangeCommands.h
@@ -66,6 +66,8 @@ namespace UE
 			BackPing,
 			RunTask,
 			NotifyEndTask,
+			QueryTaskProgress,
+			CompletedQueryTaskProgress,
 			Terminate,
 			Last
 		};
@@ -147,5 +149,39 @@ namespace UE
 			int32 TaskIndex = INDEX_NONE;
 		};
 
+		class INTERCHANGEDISPATCHER_API FQueryTaskProgressCommand : public ICommand
+		{
+		public:
+			FQueryTaskProgressCommand() = default;
+			FQueryTaskProgressCommand(const TArray<int32>& Tasks);
+			virtual ECommandId GetType() const override { return ECommandId::QueryTaskProgress; }
+
+		protected:
+			virtual void SerializeImpl(FArchive&) override;
+
+		public:
+			TArray<int32> TaskIndexes;
+		};
+
+		class INTERCHANGEDISPATCHER_API FCompletedQueryTaskProgressCommand : public ICommand
+		{
+		public:
+			FCompletedQueryTaskProgressCommand() = default;
+			virtual ECommandId GetType() const override { return ECommandId::CompletedQueryTaskProgress; }
+
+		protected:
+			virtual void SerializeImpl(FArchive&) override;
+
+		public:
+			struct FTaskProgressData
+			{
+				int32 TaskIndex;
+				ETaskState TaskState;
+				float TaskProgress;
+
+				friend FArchive& operator<<(FArchive& Ar, FTaskProgressData& A);
+			};
+			TArray<FTaskProgressData> TaskStates;
+		};
 	} //ns Interchange
 }//ns UE
diff --git a/Engine/Plugins/Runtime/Interchange/Source/Dispatcher/Public/InterchangeDispatcher.h b/Engine/Plugins/Runtime/Interchange/Source/Dispatcher/Public/InterchangeDispatcher.h
--- a/Engine/Plugins/Runtime/Interchange/Source/Dispatcher/Public/InterchangeDispatcher.h
+++ b/Engine/Plugins/Runtime/Interchange/Source/Dispatcher/Public/InterchangeDispatcher.h
@@ -23,6 +23,7 @@ namespace UE
 			int32 AddTask(const FString& JsonDescription, FInterchangeDispatcherTaskCompleted TaskCompledDelegate);
 			TOptional<FTask> GetNextTask();
 			void SetTaskState(int32 TaskIndex, ETaskState TaskState, const FString& JsonResult, const TArray<FString>& JSonMessages);
+			void GetTaskState(int32 TaskIndex, ETaskState& TaskState, double& TaskRunningStateStartTime);
 			void GetTaskState(int32 TaskIndex, ETaskState& TaskState, FString& JsonResult, TArray<FString>& JSonMessages);
 
 			void StartProcess();
@@ -35,6 +36,7 @@ namespace UE
 			void SpawnHandler();
 			bool IsHandlerAlive();
 			void CloseHandler();
+			void EmptyQueueTasks();
 
 			// Tasks
 			FCriticalSection TaskPoolCriticalSection;
diff --git a/Engine/Plugins/Runtime/Interchange/Source/Dispatcher/Public/InterchangeDispatcherTask.h b/Engine/Plugins/Runtime/Interchange/Source/Dispatcher/Public/InterchangeDispatcherTask.h
--- a/Engine/Plugins/Runtime/Interchange/Source/Dispatcher/Public/InterchangeDispatcherTask.h
+++ b/Engine/Plugins/Runtime/Interchange/Source/Dispatcher/Public/InterchangeDispatcherTask.h
@@ -42,6 +42,7 @@ namespace UE
 			FString JsonResult;
 			TArray<FString> JsonMessages;
 			FInterchangeDispatcherTaskCompleted OnTaskCompleted;
+			double RunningStateStartTime = 0;
 		};
 
 		/**
diff --git a/Engine/Plugins/Runtime/Interchange/Source/Dispatcher/Public/InterchangeWorkerHandler.h b/Engine/Plugins/Runtime/Interchange/Source/Dispatcher/Public/InterchangeWorkerHandler.h
--- a/Engine/Plugins/Runtime/Interchange/Source/Dispatcher/Public/InterchangeWorkerHandler.h
+++ b/Engine/Plugins/Runtime/Interchange/Source/Dispatcher/Public/InterchangeWorkerHandler.h
@@ -3,14 +3,13 @@
 #pragma once
 
 #include "CoreMinimal.h"
-
+#include "HAL/CriticalSection.h"
+#include "HAL/PlatformProcess.h"
+#include "HAL/Thread.h"
 #include "InterchangeCommands.h"
 #include "InterchangeDispatcherNetworking.h"
 #include "InterchangeDispatcherTask.h"
 
-#include "HAL/PlatformProcess.h"
-#include "HAL/Thread.h"
-
 namespace UE
 {
 	namespace Interchange
@@ -53,6 +52,7 @@ namespace UE
 			void Stop();
 			void StopBlocking();
 
+			FSimpleMulticastDelegate OnWorkerHandlerExitLoop;
 		protected:
 			void ProcessCommand(ICommand& Command);
 		private:
@@ -62,8 +62,11 @@ namespace UE
 
 			void ProcessCommand(FPingCommand& PingCommand);
 			void ProcessCommand(FCompletedTaskCommand& RunTaskCommand);
+			void ProcessCommand(FCompletedQueryTaskProgressCommand& CompletedQueryTaskProgressCommand);
 			const TCHAR* EWorkerErrorStateAsString(EWorkerErrorState e);
 
+			void KillAllCurrentTasks();
+
 		private:
 			FInterchangeDispatcher& Dispatcher;
 
@@ -80,8 +83,10 @@ namespace UE
 
 			// self
 			FString ResultFolder;
+			FCriticalSection CurrentTasksLock;
 			TArray<int32> CurrentTasks;
 			bool bShouldTerminate;
+			double LastProgressMessageTime;
 
 			friend Dispatcher::FTaskProcessCommand;
 		};
diff --git a/Engine/Plugins/Runtime/Interchange/Source/Import/Private/Fbx/InterchangeFbxTranslator.cpp b/Engine/Plugins/Runtime/Interchange/Source/Import/Private/Fbx/InterchangeFbxTranslator.cpp
--- a/Engine/Plugins/Runtime/Interchange/Source/Import/Private/Fbx/InterchangeFbxTranslator.cpp
+++ b/Engine/Plugins/Runtime/Interchange/Source/Import/Private/Fbx/InterchangeFbxTranslator.cpp
@@ -162,8 +162,8 @@ TFuture<TOptional<UE::Interchange::FSkeletalMeshLodPayloadData>> UInterchangeFbx
 	}
 
 	//Create a json command to read the fbx file
-	FString JsonCommand = CreateFetchPayloadFbxCommand(PayLoadKey);
-	int32 TaskIndex = Dispatcher->AddTask(JsonCommand, FInterchangeDispatcherTaskCompleted::CreateLambda([this, Promise](const int32 TaskIndex)
+	const FString JsonCommand = CreateFetchPayloadFbxCommand(PayLoadKey);
+	const int32 CreatedTaskIndex = Dispatcher->AddTask(JsonCommand, FInterchangeDispatcherTaskCompleted::CreateLambda([this, Promise](const int32 TaskIndex)
 	{
 		UE::Interchange::ETaskState TaskState;
 		FString JsonResult;
@@ -216,6 +216,12 @@ TFuture<TOptional<UE::Interchange::FSkeletalMeshLodPayloadData>> UInterchangeFbx
 		Promise->SetValue(MoveTemp(SkeletalMeshLodPayload));
 	}));
 
+	//The task was not added to the dispatcher
+	if (CreatedTaskIndex == INDEX_NONE)
+	{
+		Promise->SetValue(TOptional<UE::Interchange::FSkeletalMeshLodPayloadData>{});
+	}
+
 	return Promise->GetFuture();
 }
 
@@ -229,8 +235,7 @@ TFuture<TOptional<UE::Interchange::FSkeletalMeshBlendShapePayloadData>> UInterch
 	}
 
 	//Create a json command to read the fbx file
-	FString JsonCommand = CreateFetchPayloadFbxCommand(PayLoadKey);
-
+	const FString JsonCommand = CreateFetchPayloadFbxCommand(PayLoadKey);
 	const int32 CreatedTaskIndex = Dispatcher->AddTask(JsonCommand, FInterchangeDispatcherTaskCompleted::CreateLambda([this, Promise](const int32 TaskIndex)
 	{
 		UE::Interchange::ETaskState TaskState;
@@ -275,6 +280,12 @@ TFuture<TOptional<UE::Interchange::FSkeletalMeshBlendShapePayloadData>> UInterch
 		Promise->SetValue(MoveTemp(SkeletalMeshBlendShapePayload));
 	}));
 
+	//The task was not added to the dispatcher
+	if (CreatedTaskIndex == INDEX_NONE)
+	{
+		Promise->SetValue(TOptional<UE::Interchange::FSkeletalMeshBlendShapePayloadData>{});
+	}
+
 	return Promise->GetFuture();
 }
 
diff --git a/Engine/Source/Programs/InterchangeWorker/Private/InterchangeWorkerImpl.cpp b/Engine/Source/Programs/InterchangeWorker/Private/InterchangeWorkerImpl.cpp
--- a/Engine/Source/Programs/InterchangeWorker/Private/InterchangeWorkerImpl.cpp
+++ b/Engine/Source/Programs/InterchangeWorker/Private/InterchangeWorkerImpl.cpp
@@ -78,6 +78,9 @@ bool FInterchangeWorkerImpl::Run()
 					));
 				}
 					break;
+				case ECommandId::QueryTaskProgress:
+					ProcessCommand(*StaticCast<FQueryTaskProgressCommand*>(Command.Get()));
+					break;
 
 				case ECommandId::Terminate:
 					UE_LOG(LogInterchangeWorker, Verbose, TEXT("Terminate command received. Exiting."));
@@ -151,6 +154,24 @@ void FInterchangeWorkerImpl::ProcessCommand(const FBackPingCommand& BackPingComm
 	PingStartCycle = 0;
 }
 
+void FInterchangeWorkerImpl::ProcessCommand(const UE::Interchange::FQueryTaskProgressCommand& QueryTaskProgressCommand)
+{
+	FCompletedQueryTaskProgressCommand CompletedCommand;
+	CompletedCommand.TaskStates.AddDefaulted(QueryTaskProgressCommand.TaskIndexes.Num());
+	{
+		FScopeLock Lock(&TFinishThreadCriticalSection);
+		for (int32 ProgressTaskIndex = 0; ProgressTaskIndex < QueryTaskProgressCommand.TaskIndexes.Num(); ++ProgressTaskIndex)
+		{
+			CompletedCommand.TaskStates[ProgressTaskIndex].TaskIndex = QueryTaskProgressCommand.TaskIndexes[ProgressTaskIndex];
+			//If we do not have any job to process, set the state to unknown. The caller will know we are not currently processing this task, whihc mean we never receive this task or this task is completed.
+			CompletedCommand.TaskStates[ProgressTaskIndex].TaskState = ActiveThreads.Num() == 0 ? ETaskState::Unknown : ETaskState::Running;
+			//TODO: implement better progress report, we currently always return 0.0
+			CompletedCommand.TaskStates[ProgressTaskIndex].TaskProgress = 0.0f;
+		}
+	}
+	CommandIO.SendCommand(CompletedCommand, Config::SendCommandTimeout_s);
+}
+
 void FInterchangeWorkerImpl::ProcessCommand(const TSharedPtr<UE::Interchange::ICommand> Command, const FString& ThreadName)
 {
 	if (!Command.IsValid())
diff --git a/Engine/Source/Programs/InterchangeWorker/Private/InterchangeWorkerImpl.h b/Engine/Source/Programs/InterchangeWorker/Private/InterchangeWorkerImpl.h
--- a/Engine/Source/Programs/InterchangeWorker/Private/InterchangeWorkerImpl.h
+++ b/Engine/Source/Programs/InterchangeWorker/Private/InterchangeWorkerImpl.h
@@ -23,6 +23,7 @@ private:
 	void ProcessCommand(const UE::Interchange::FPingCommand& PingCommand);
 	void ProcessCommand(const UE::Interchange::FBackPingCommand& BackPingCommand);
 	void ProcessCommand(const TSharedPtr<UE::Interchange::ICommand> Command, const FString& ThreadName);
+	void ProcessCommand(const UE::Interchange::FQueryTaskProgressCommand& QueryTaskProgressCommand);
 
 	UE::Interchange::ETaskState LoadFbxFile(const UE::Interchange::FJsonLoadSourceCmd& LoadSourceCommand, FString& OutJSonResult, TArray<FString>& OutJSonMessages);
 	UE::Interchange::ETaskState FetchFbxPayload(const UE::Interchange::FJsonFetchPayloadCmd& FetchPayloadCommand, FString& OutJSonResult, TArray<FString>& OutJSonMessages);
