Fix non unity build

diff --git a/Engine/Plugins/Runtime/Interchange/Source/Dispatcher/Private/InterchangeWorkerHandler.cpp b/Engine/Plugins/Runtime/Interchange/Source/Dispatcher/Private/InterchangeWorkerHandler.cpp
--- a/Engine/Plugins/Runtime/Interchange/Source/Dispatcher/Private/InterchangeWorkerHandler.cpp
+++ b/Engine/Plugins/Runtime/Interchange/Source/Dispatcher/Private/InterchangeWorkerHandler.cpp
@@ -10,6 +10,7 @@
 #include "HAL/FileManager.h"
 #include "HAL/PlatformProcess.h"
 #include "Misc/Paths.h"
+#include "Misc/ScopeLock.h"
 #include "Sockets.h"
 #include "SocketSubsystem.h"
 
