Added texture arrays support to interchange.

Added support to import texture array from DDS File.

The texture 2d array factory is still missing some feature such as creating a empty texture and setting the SourceTextures.
Those will be done in another pass for some other jira.

diff --git a/Engine/Plugins/Runtime/Interchange/Source/Import/Private/Texture/InterchangeDDSTranslator.cpp b/Engine/Plugins/Runtime/Interchange/Source/Import/Private/Texture/InterchangeDDSTranslator.cpp
--- a/Engine/Plugins/Runtime/Interchange/Source/Import/Private/Texture/InterchangeDDSTranslator.cpp
+++ b/Engine/Plugins/Runtime/Interchange/Source/Import/Private/Texture/InterchangeDDSTranslator.cpp
@@ -97,7 +97,7 @@ bool UInterchangeDDSTranslator::Translate(const UInterchangeSourceData* SourceDa
 	}
 
 
-	FDDSLoadHelper  DDSLoadHelper(HeaderDataBuffer.GetData(), HeaderDataBuffer.Num());
+	FDDSLoadHelper DDSLoadHelper(HeaderDataBuffer.GetData(), HeaderDataBuffer.Num());
 	if (DDSLoadHelper.IsValid2DTexture())
 	{
 		return UE::Interchange::FTextureTranslatorUtilities::Generic2DTextureTranslate(SourceData, BaseNodeContainer);
@@ -105,7 +105,12 @@ bool UInterchangeDDSTranslator::Translate(const UInterchangeSourceData* SourceDa
 
 	if (DDSLoadHelper.IsValidCubemapTexture())
 	{
-		return UE::Interchange::FTextureTranslatorUtilities::GenericCubeTextureTranslate(SourceData, BaseNodeContainer);
+		return UE::Interchange::FTextureTranslatorUtilities::GenericTextureCubeTranslate(SourceData, BaseNodeContainer);
+	}
+
+	if (DDSLoadHelper.IsValidArrayTexture())
+	{
+		return UE::Interchange::FTextureTranslatorUtilities::GenericTexture2DArrayTranslate(SourceData, BaseNodeContainer);
 	}
 
 	return false;
diff --git a/Engine/Plugins/Runtime/Interchange/Source/Import/Private/Texture/InterchangeTextureFactory.cpp b/Engine/Plugins/Runtime/Interchange/Source/Import/Private/Texture/InterchangeTextureFactory.cpp
--- a/Engine/Plugins/Runtime/Interchange/Source/Import/Private/Texture/InterchangeTextureFactory.cpp
+++ b/Engine/Plugins/Runtime/Interchange/Source/Import/Private/Texture/InterchangeTextureFactory.cpp
@@ -4,10 +4,14 @@
 
 #include "Engine/Texture.h"
 #include "Engine/Texture2D.h"
+#include "Engine/Texture2DArray.h"
 #include "Engine/TextureCube.h"
 #include "InterchangeAssetImportData.h"
 #include "InterchangeImportCommon.h"
 #include "InterchangeImportLog.h"
+#include "InterchangeTexture2DArrayFactoryNode.h"
+#include "InterchangeTexture2DArrayNode.h"
+#include "InterchangeTexture2DNode.h"
 #include "InterchangeTextureCubeFactoryNode.h"
 #include "InterchangeTextureCubeNode.h"
 #include "InterchangeTextureFactoryNode.h"
@@ -38,20 +42,27 @@ namespace UE::Interchange::Private::InterchangeTextureFactory
 	{
 		UClass* TextureCubeFactoryClass = UInterchangeTextureCubeFactoryNode::StaticClass();
 		UClass* TextureFactoryClass = UInterchangeTextureFactoryNode::StaticClass();
+		UClass* Texture2DArrayFactoryClass = UInterchangeTexture2DArrayFactoryNode::StaticClass();
 		UClass* AssetClass = AssetNode->GetClass();
 #if USTRUCT_FAST_ISCHILDOF_IMPL == USTRUCT_ISCHILDOF_STRUCTARRAY
-		if(AssetClass->IsChildOf(TextureFactoryClass))
+		if (AssetClass->IsChildOf(Texture2DArrayFactoryClass))
+		{
+			return Texture2DArrayFactoryClass;
+		}
+		else if (AssetClass->IsChildOf(TextureFactoryClass))
 		{
 			return TextureFactoryClass;
 		}
-		else if(AssetClass->IsChildOf(TextureCubeFactoryClass))
+		else if (AssetClass->IsChildOf(TextureCubeFactoryClass))
 		{
 			return TextureCubeFactoryClass;
 		}
 #else
 		while (AssetClass)
 		{
-			if (AssetClass == TextureCubeFactoryClass || AssetClass == TextureFactoryClass)
+			if (AssetClass == TextureCubeFactoryClass
+				|| AssetClass == TextureFactoryClass
+				|| AssetClass == Texture2DArrayFactoryClass)
 			{
 				return AssetClass;
 			}
@@ -63,7 +74,10 @@ namespace UE::Interchange::Private::InterchangeTextureFactory
 		return nullptr;
 	}
 
-	using FTextureFactoryNodeVariant = TVariant<FEmptyVariantState, UInterchangeTextureFactoryNode*, UInterchangeTextureCubeFactoryNode*>; 
+	using FTextureFactoryNodeVariant = TVariant<FEmptyVariantState
+		, UInterchangeTextureFactoryNode*
+		, UInterchangeTextureCubeFactoryNode*
+		, UInterchangeTexture2DArrayFactoryNode*>; 
 
 	FTextureFactoryNodeVariant GetAsTextureFactoryNodeVariant(UInterchangeBaseNode* AssetNode, UClass* SupportedFactoryNodeClass)
 	{
@@ -83,12 +97,20 @@ namespace UE::Interchange::Private::InterchangeTextureFactory
 			{
 				return FTextureFactoryNodeVariant(TInPlaceType<UInterchangeTextureCubeFactoryNode*>(), static_cast<UInterchangeTextureCubeFactoryNode*>(AssetNode));
 			}
+
+			if (SupportedFactoryNodeClass == UInterchangeTexture2DArrayFactoryNode::StaticClass())
+			{
+				return FTextureFactoryNodeVariant(TInPlaceType<UInterchangeTexture2DArrayFactoryNode*>(), static_cast<UInterchangeTexture2DArrayFactoryNode*>(AssetNode));
+			}
 		}
 
 		return {};
 	}
 
-	using FTextureNodeVariant = TVariant<FEmptyVariantState, const UInterchangeTextureNode*, const UInterchangeTextureCubeNode*>; 
+	using FTextureNodeVariant = TVariant<FEmptyVariantState
+		, const UInterchangeTexture2DNode*
+		, const UInterchangeTextureCubeNode*
+		, const UInterchangeTexture2DArrayNode*>; 
 
 	FTextureNodeVariant GetTextureNodeVariantFromFactoryVariant(const FTextureFactoryNodeVariant& FactoryVariant, const UInterchangeBaseNodeContainer* NodeContainer)
 	{
@@ -98,9 +120,9 @@ namespace UE::Interchange::Private::InterchangeTextureFactory
 		{
 			if ((*TextureFactoryNode)->GetCustomTranslatedTextureNodeUid(TextureNodeUniqueID))
 			{
-				if (const UInterchangeTextureNode* TextureTranslatedNode = Cast<UInterchangeTextureNode>(NodeContainer->GetNode(TextureNodeUniqueID)))
+				if (const UInterchangeTexture2DNode* TextureTranslatedNode = Cast<UInterchangeTexture2DNode>(NodeContainer->GetNode(TextureNodeUniqueID)))
 				{
-					return FTextureNodeVariant(TInPlaceType<const UInterchangeTextureNode*>(), TextureTranslatedNode);
+					return FTextureNodeVariant(TInPlaceType<const UInterchangeTexture2DNode*>(), TextureTranslatedNode);
 				}
 			}
 		}
@@ -114,20 +136,30 @@ namespace UE::Interchange::Private::InterchangeTextureFactory
 					{
 						return FTextureNodeVariant(TInPlaceType<const UInterchangeTextureCubeNode*>(), TextureCubeTranslatedNode);
 					}
-					else if (const UInterchangeTextureNode* TextureTranslatedNode = Cast<UInterchangeTextureNode>(Node))
+					else if (const UInterchangeTexture2DNode* TextureTranslatedNode = Cast<UInterchangeTexture2DNode>(Node))
 					{
-						return FTextureNodeVariant(TInPlaceType<const UInterchangeTextureNode*>(), TextureTranslatedNode);
+						return FTextureNodeVariant(TInPlaceType<const UInterchangeTexture2DNode*>(), TextureTranslatedNode);
 					}
 				}
 			}
 		}
+		else if (UInterchangeTexture2DArrayFactoryNode* const* Texture2DArrayFactoryNode = FactoryVariant.TryGet<UInterchangeTexture2DArrayFactoryNode*>())
+		{
+			if ((*Texture2DArrayFactoryNode)->GetCustomTranslatedTextureNodeUid(TextureNodeUniqueID))
+			{
+				if (const UInterchangeTexture2DArrayNode* Texture2DArrayTranslatedNode = Cast<UInterchangeTexture2DArrayNode>(NodeContainer->GetNode(TextureNodeUniqueID)))
+				{
+					return FTextureNodeVariant(TInPlaceType<const UInterchangeTexture2DArrayNode*>(), Texture2DArrayTranslatedNode);
+				}
+			}
+		}
 
 		return {};
 	}
 
 	bool HasPayloadKey(const FTextureNodeVariant& TextureNodeVariant)
 	{
-		if (const UInterchangeTextureNode* const* TextureNode =  TextureNodeVariant.TryGet<const UInterchangeTextureNode*>())
+		if (const UInterchangeTexture2DNode* const* TextureNode =  TextureNodeVariant.TryGet<const UInterchangeTexture2DNode*>())
 		{
 			return (*TextureNode)->GetPayLoadKey().IsSet();
 		}
@@ -137,12 +169,17 @@ namespace UE::Interchange::Private::InterchangeTextureFactory
 			return (*TextureCubeNode)->GetPayLoadKey().IsSet();
 		}
 
+		if (const UInterchangeTexture2DArrayNode* const* Texture2DArrayNode = TextureNodeVariant.TryGet<const UInterchangeTexture2DArrayNode*>())
+		{
+			return (*Texture2DArrayNode)->GetPayLoadKey().IsSet();
+		}
+
 	return false;
 	}
 
 	TOptional<FString> GetPayloadKey(const FTextureNodeVariant& TextureNodeVariant)
 	{
-		if (const UInterchangeTextureNode* const* TextureNode =  TextureNodeVariant.TryGet<const UInterchangeTextureNode*>())
+		if (const UInterchangeTexture2DNode* const* TextureNode =  TextureNodeVariant.TryGet<const UInterchangeTexture2DNode*>())
 		{
 			return (*TextureNode)->GetPayLoadKey();
 		}
@@ -152,6 +189,11 @@ namespace UE::Interchange::Private::InterchangeTextureFactory
 			return (*TextureCubeNode)->GetPayLoadKey();
 		}
 
+		if (const UInterchangeTexture2DArrayNode* const* Texture2DArrayNode = TextureNodeVariant.TryGet<const UInterchangeTexture2DArrayNode*>())
+		{
+			return (*Texture2DArrayNode)->GetPayLoadKey();
+		}
+
 		return {};
 	}
 
@@ -160,7 +202,7 @@ namespace UE::Interchange::Private::InterchangeTextureFactory
 	FTexturePayloadVariant GetTexturePayload(const UInterchangeSourceData* SourceData, const FString& PayloadKey, const FTextureNodeVariant& TextureNodeVariant, const UInterchangeTranslatorBase* Translator)
 	{
 		// Standard texture 2D payload
-		if (const UInterchangeTextureNode* const* TextureNode =  TextureNodeVariant.TryGet<const UInterchangeTextureNode*>())
+		if (const UInterchangeTexture2DNode* const* TextureNode =  TextureNodeVariant.TryGet<const UInterchangeTexture2DNode*>())
 		{
 			if (const IInterchangeTexturePayloadInterface* TextureTranslator = Cast<IInterchangeTexturePayloadInterface>(Translator))
 			{
@@ -172,8 +214,8 @@ namespace UE::Interchange::Private::InterchangeTextureFactory
 			}
 		}
 
-		// Cube texture payload
-		if (TextureNodeVariant.IsType<const UInterchangeTextureCubeNode*>())
+		// Cube or array texture payload 
+		if (TextureNodeVariant.IsType<const UInterchangeTextureCubeNode*>() || TextureNodeVariant.IsType<const UInterchangeTexture2DArrayNode*>())
 		{
 			if (const IInterchangeSlicedTexturePayloadInterface* SlicedTextureTranslator = Cast<IInterchangeSlicedTexturePayloadInterface>(Translator))
 			{
@@ -279,6 +321,50 @@ namespace UE::Interchange::Private::InterchangeTextureFactory
 		return false;
 	}
 
+	void SetupTextureSourceData(UTexture* Texture, const FImportSlicedImage& SlicedImage)
+	{
+		Texture->Source.Init(
+			SlicedImage.SizeX,
+			SlicedImage.SizeY,
+			SlicedImage.NumSlice,
+			SlicedImage.NumMips,
+			SlicedImage.Format
+		);
+
+		Texture->CompressionSettings = SlicedImage.CompressionSettings;
+		Texture->SRGB = SlicedImage.bSRGB;
+
+
+		uint8* DestMipData[MAX_TEXTURE_MIP_COUNT] = {0};
+		int32 MipSize[MAX_TEXTURE_MIP_COUNT] = {0};
+		for (int32 MipIndex = 0; MipIndex < SlicedImage.NumMips; ++MipIndex)
+		{
+			DestMipData[MipIndex] = Texture->Source.LockMip(MipIndex);
+			MipSize[MipIndex] = Texture->Source.CalcMipSize(MipIndex) / SlicedImage.NumSlice;
+		}
+
+		for (int32 SliceIndex = 0; SliceIndex < SlicedImage.NumSlice; ++SliceIndex)
+		{
+			const uint8* SrcMipData = SlicedImage.GetMipData(0, SliceIndex);
+			for (int32 MipIndex = 0; MipIndex < SlicedImage.NumMips; ++MipIndex)
+			{
+				FMemory::Memcpy(DestMipData[MipIndex] + MipSize[MipIndex] * SliceIndex, SrcMipData, MipSize[MipIndex]);
+				SrcMipData += MipSize[MipIndex];
+			}
+		}
+
+		for (int32 MipIndex = 0; MipIndex < SlicedImage.NumMips; ++MipIndex)
+		{
+			Texture->Source.UnlockMip(MipIndex);
+		}
+
+		if (SlicedImage.MipGenSettings.IsSet())
+		{
+			// if the source has mips we keep the mips by default, unless the user changes that
+			Texture->MipGenSettings = SlicedImage.MipGenSettings.GetValue();
+		}
+	}
+
 	bool SetupTextureCubeSourceData(UTextureCube* TextureCube, const FTexturePayloadVariant& TexturePayload)
 	{
 		if (const TOptional<FImportSlicedImage>* OptionalSlicedImage = TexturePayload.TryGet<TOptional<FImportSlicedImage>>())
@@ -290,47 +376,7 @@ namespace UE::Interchange::Private::InterchangeTextureFactory
 				// Cube texture always have six slice
 				if (SlicedImage.NumSlice == 6)
 				{
-					TextureCube->Source.Init(
-						SlicedImage.SizeX,
-						SlicedImage.SizeY,
-						SlicedImage.NumSlice,
-						SlicedImage.NumMips,
-						SlicedImage.Format
-					);
-
-					TextureCube->CompressionSettings = SlicedImage.CompressionSettings;
-					TextureCube->SRGB = SlicedImage.bSRGB;
-
-
-					uint8* DestMipData[MAX_TEXTURE_MIP_COUNT] = {0};
-					int32 MipSize[MAX_TEXTURE_MIP_COUNT] = {0};
-					for (int32 MipIndex = 0; MipIndex < SlicedImage.NumMips; ++MipIndex)
-					{
-						DestMipData[MipIndex] = TextureCube->Source.LockMip(MipIndex);
-						MipSize[MipIndex] = TextureCube->Source.CalcMipSize(MipIndex) / 6;
-					}
-
-					for (int32 SliceIndex = 0; SliceIndex < 6; ++SliceIndex)
-					{
-						const uint8* SrcMipData = SlicedImage.GetMipData(0, SliceIndex);
-						for (int32 MipIndex = 0; MipIndex < SlicedImage.NumMips; ++MipIndex)
-						{
-							FMemory::Memcpy(DestMipData[MipIndex] + MipSize[MipIndex] * SliceIndex, SrcMipData, MipSize[MipIndex]);
-							SrcMipData += MipSize[MipIndex];
-						}
-					}
-
-					for (int32 MipIndex = 0; MipIndex < SlicedImage.NumMips; ++MipIndex)
-					{
-						TextureCube->Source.UnlockMip(MipIndex);
-					}
-
-					if (SlicedImage.MipGenSettings.IsSet())
-					{
-						// if the source has mips we keep the mips by default, unless the user changes that
-						TextureCube->MipGenSettings = SlicedImage.MipGenSettings.GetValue();
-					}
-
+					SetupTextureSourceData(TextureCube, SlicedImage);
 					return true;
 				}
 			}
@@ -341,7 +387,31 @@ namespace UE::Interchange::Private::InterchangeTextureFactory
 			{
 				const FImportImage& Image = OptionalImage->GetValue();
 				SetupTextureSourceData(TextureCube, Image);
+				return true;
+			}
+		}
+
+		return false;
+	}
+
+	bool SetupTexture2DArraySourceData(UTexture2DArray* Texture2DArray, const FTexturePayloadVariant& TexturePayload)
+	{
+		if (const TOptional<FImportSlicedImage>* OptionalSlicedImage = TexturePayload.TryGet<TOptional<FImportSlicedImage>>())
+		{
+			if (OptionalSlicedImage->IsSet())
+			{
+				const FImportSlicedImage& SlicedImage = OptionalSlicedImage->GetValue();
 
+				SetupTextureSourceData(Texture2DArray, SlicedImage);
+				return true;
+			}
+		}
+		else if (const TOptional<FImportImage>* OptionalImage = TexturePayload.TryGet<TOptional<FImportImage>>())
+		{
+			if (OptionalImage->IsSet())
+			{
+				const FImportImage& Image = OptionalImage->GetValue();
+				SetupTextureSourceData(Texture2DArray, Image);
 				return true;
 			}
 		}
@@ -470,7 +540,7 @@ UObject* UInterchangeTextureFactory::CreateAsset(const UInterchangeTextureFactor
 	const TOptional<FString>& PayLoadKey = GetPayloadKey(TextureNodeVariant);
 	if (!PayLoadKey.IsSet())
 	{
-		UE_LOG(LogInterchangeImport, Error, TEXT("UInterchangeTextureFactory: Texture translated node (UInterchangeTextureNode) doesnt have a payload key."));
+		UE_LOG(LogInterchangeImport, Error, TEXT("UInterchangeTextureFactory: Texture translated node (UInterchangeTexture2DNode) doesnt have a payload key."));
 		return nullptr;
 	}
 
@@ -524,9 +594,17 @@ UObject* UInterchangeTextureFactory::CreateAsset(const UInterchangeTextureFactor
 			return nullptr;
 		}
 	}
+	else if (UTexture2DArray* Texture2DArray = Cast<UTexture2DArray>(Texture))
+	{
+		if (!SetupTexture2DArraySourceData(Texture2DArray, TexturePayload))
+		{
+			UE_LOG(LogInterchangeImport, Error, TEXT("UInterchangeTextureFactory: The Payload was invalid for a TextureCube"), *Arguments.AssetName);
+			return nullptr;
+		}
+	}
 	else
 	{
-		//The texture is not supported
+		// The texture is not supported
 		Texture->RemoveFromRoot();
 		Texture->MarkPendingKill();
 		return nullptr;
diff --git a/Engine/Plugins/Runtime/Interchange/Source/Import/Private/Texture/InterchangeUDIMTranslator.cpp b/Engine/Plugins/Runtime/Interchange/Source/Import/Private/Texture/InterchangeUDIMTranslator.cpp
--- a/Engine/Plugins/Runtime/Interchange/Source/Import/Private/Texture/InterchangeUDIMTranslator.cpp
+++ b/Engine/Plugins/Runtime/Interchange/Source/Import/Private/Texture/InterchangeUDIMTranslator.cpp
@@ -3,7 +3,7 @@
 #include "Texture/InterchangeUDIMTranslator.h"
 
 #include "InterchangeManager.h"
-#include "InterchangeTextureNode.h"
+#include "InterchangeTexture2DNode.h"
 #include "Texture/InterchangeTexturePayloadInterface.h"
 
 
@@ -168,7 +168,7 @@ bool UInterchangeUDIMTranslator::Translate(const UInterchangeSourceData* SourceD
 		return false;
 	}
 
-	UInterchangeTextureNode* TextureNode = NewObject<UInterchangeTextureNode>(&BaseNodeContainer, NAME_None);
+	UInterchangeTexture2DNode* TextureNode = NewObject<UInterchangeTexture2DNode>(&BaseNodeContainer, NAME_None);
 
 	if (!ensure(TextureNode))
 	{
diff --git a/Engine/Plugins/Runtime/Interchange/Source/Import/Private/Texture/TextureTranslatorUtilities.cpp b/Engine/Plugins/Runtime/Interchange/Source/Import/Private/Texture/TextureTranslatorUtilities.cpp
--- a/Engine/Plugins/Runtime/Interchange/Source/Import/Private/Texture/TextureTranslatorUtilities.cpp
+++ b/Engine/Plugins/Runtime/Interchange/Source/Import/Private/Texture/TextureTranslatorUtilities.cpp
@@ -3,60 +3,61 @@
 
 #include "CoreMinimal.h"
 #include "InterchangeSourceData.h"
+#include "InterchangeTexture2DArrayNode.h"
+#include "InterchangeTexture2DNode.h"
 #include "InterchangeTextureCubeNode.h"
 #include "InterchangeTextureNode.h"
 #include "Misc/Paths.h"
 #include "Nodes/InterchangeBaseNodeContainer.h"
+#include "Templates/SubclassOf.h"
 #include "UObject/Object.h"
 #include "UObject/ObjectMacros.h"
 
-bool UE::Interchange::FTextureTranslatorUtilities::Generic2DTextureTranslate(const UInterchangeSourceData* SourceData, UInterchangeBaseNodeContainer& BaseNodeContainer)
+namespace UE::Interchange::Private
 {
-	FString Filename = SourceData->GetFilename();
-	FPaths::NormalizeFilename(Filename);
-	if (!FPaths::FileExists(Filename))
+	bool GenericTextureTranslate(const UInterchangeSourceData* SourceData, UInterchangeBaseNodeContainer& BaseNodeContainer, const TSubclassOf<UInterchangeTextureNode>& TextureNodeClass)
 	{
-		return false;
+		FString Filename = SourceData->GetFilename();
+		FPaths::NormalizeFilename(Filename);
+		if (!FPaths::FileExists(Filename))
+		{
+			return false;
+		}
+
+		UClass* Class = TextureNodeClass.Get();
+		if (!ensure(Class))
+		{
+			return false;
+		}
+
+		FString DisplayLabel = FPaths::GetBaseFilename(Filename);
+		FString NodeUID(Filename);
+		UInterchangeTextureNode* TextureNode = NewObject<UInterchangeTextureNode>(&BaseNodeContainer, Class);
+		if (!ensure(TextureNode))
+		{
+			return false;
+		}
+
+		TextureNode->InitializeNode(NodeUID, DisplayLabel, EInterchangeNodeContainerType::NodeContainerType_TranslatedAsset);
+		TextureNode->SetPayLoadKey(Filename);
+
+		BaseNodeContainer.AddNode(TextureNode);
+
+		return true;
 	}
-
-	FString DisplayLabel = FPaths::GetBaseFilename(Filename);
-	FString NodeUID(Filename);
-	UInterchangeTextureNode* TextureNode = NewObject<UInterchangeTextureNode>(&BaseNodeContainer, NAME_None);
-	if (!ensure(TextureNode))
-	{
-		return false;
-	}
-	//Creating a UTexture2D
-	TextureNode->InitializeNode(NodeUID, DisplayLabel, EInterchangeNodeContainerType::NodeContainerType_TranslatedAsset);
-	TextureNode->SetPayLoadKey(Filename);
-
-	BaseNodeContainer.AddNode(TextureNode);
-
-	return true;
 }
 
-bool UE::Interchange::FTextureTranslatorUtilities::GenericCubeTextureTranslate(const UInterchangeSourceData* SourceData, UInterchangeBaseNodeContainer& BaseNodeContainer)
+bool UE::Interchange::FTextureTranslatorUtilities::Generic2DTextureTranslate(const UInterchangeSourceData* SourceData, UInterchangeBaseNodeContainer& BaseNodeContainer)
 {
-	FString Filename = SourceData->GetFilename();
-	FPaths::NormalizeFilename(Filename);
-	if (!FPaths::FileExists(Filename))
-	{
-		return false;
-	}
-
-	FString DisplayLabel = FPaths::GetBaseFilename(Filename);
-	FString NodeUID(Filename);
-	UInterchangeTextureCubeNode* TextureCubeNode = NewObject<UInterchangeTextureCubeNode>(&BaseNodeContainer, NAME_None);
-	if (!ensure(TextureCubeNode))
-	{
-		return false;
-	}
-
-	//Creating a Cube Texture
-	TextureCubeNode->InitializeNode(NodeUID, DisplayLabel, EInterchangeNodeContainerType::NodeContainerType_TranslatedAsset);
-	TextureCubeNode->SetPayLoadKey(Filename);
+	return Private::GenericTextureTranslate(SourceData, BaseNodeContainer, UInterchangeTexture2DNode::StaticClass());
+}
 
-	BaseNodeContainer.AddNode(TextureCubeNode);
+bool UE::Interchange::FTextureTranslatorUtilities::GenericTextureCubeTranslate(const UInterchangeSourceData* SourceData, UInterchangeBaseNodeContainer& BaseNodeContainer)
+{
+	return Private::GenericTextureTranslate(SourceData, BaseNodeContainer, UInterchangeTextureCubeNode::StaticClass());
+}
 
-	return true;
+bool UE::Interchange::FTextureTranslatorUtilities::GenericTexture2DArrayTranslate(const UInterchangeSourceData* SourceData, UInterchangeBaseNodeContainer& BaseNodeContainer)
+{
+	return Private::GenericTextureTranslate(SourceData, BaseNodeContainer, UInterchangeTexture2DArrayNode::StaticClass());
 }
diff --git a/Engine/Plugins/Runtime/Interchange/Source/Import/Private/Texture/TextureTranslatorUtilities.h b/Engine/Plugins/Runtime/Interchange/Source/Import/Private/Texture/TextureTranslatorUtilities.h
--- a/Engine/Plugins/Runtime/Interchange/Source/Import/Private/Texture/TextureTranslatorUtilities.h
+++ b/Engine/Plugins/Runtime/Interchange/Source/Import/Private/Texture/TextureTranslatorUtilities.h
@@ -7,6 +7,11 @@
 #include "UObject/Object.h"
 #include "UObject/ObjectMacros.h"
 
+class UInterchangeTextureNode;
+
+template <class T>
+class TSubclassOf;
+
 namespace UE
 {
 	namespace Interchange
@@ -16,7 +21,9 @@ namespace UE
 		public:
 			static bool Generic2DTextureTranslate(const UInterchangeSourceData* SourceData, UInterchangeBaseNodeContainer& BaseNodeContainer);
 
-			static bool GenericCubeTextureTranslate(const UInterchangeSourceData* SourceData, UInterchangeBaseNodeContainer& BaseNodeContainer);
+			static bool GenericTextureCubeTranslate(const UInterchangeSourceData* SourceData, UInterchangeBaseNodeContainer& BaseNodeContainer);
+
+			static bool GenericTexture2DArrayTranslate(const UInterchangeSourceData* SourceData, UInterchangeBaseNodeContainer& BaseNodeContainer);
 		};
 	}//ns Interchange
 }//ns UE
diff --git a/Engine/Plugins/Runtime/Interchange/Source/Nodes/Public/InterchangeTexture2DArrayFactoryNode.h b/Engine/Plugins/Runtime/Interchange/Source/Nodes/Public/InterchangeTexture2DArrayFactoryNode.h
new file mode 100644
--- /dev/null
+++ b/Engine/Plugins/Runtime/Interchange/Source/Nodes/Public/InterchangeTexture2DArrayFactoryNode.h
@@ -0,0 +1,81 @@
+// Copyright Epic Games, Inc. All Rights Reserved.
+
+#pragma once
+
+#include "CoreMinimal.h"
+#include "InterchangeTextureFactoryNode.h"
+#include "Nodes/InterchangeBaseNode.h"
+
+#if WITH_ENGINE
+#include "Engine/Texture.h"
+#include "Engine/Texture2DArray.h"
+#endif
+
+
+#include "InterchangeTexture2DArrayFactoryNode.generated.h"
+
+UCLASS(BlueprintType)
+class INTERCHANGENODES_API UInterchangeTexture2DArrayFactoryNode : public UInterchangeTextureFactoryNode
+{
+	GENERATED_BODY()
+
+public:
+
+	virtual bool SetCustomAddressX(const uint8 AttributeValue, bool bAddApplyDelegate = true) override
+	{
+#if WITH_EDITORONLY_DATA
+		IMPLEMENT_NODE_ATTRIBUTE_SETTER(UInterchangeTexture2DArrayFactoryNode, AddressX, uint8, UTexture2DArray)
+#else
+		IMPLEMENT_NODE_ATTRIBUTE_SETTER_NODELEGATE(AddressX, uint8)
+#endif
+	}
+
+	virtual bool SetCustomAddressY(const uint8 AttributeValue, bool bAddApplyDelegate = true) override
+	{
+#if WITH_EDITORONLY_DATA
+		IMPLEMENT_NODE_ATTRIBUTE_SETTER(UInterchangeTexture2DArrayFactoryNode, AddressY, uint8, UTexture2DArray)
+#else
+		IMPLEMENT_NODE_ATTRIBUTE_SETTER_NODELEGATE(AddressY, uint8)
+#endif
+	}
+
+	UFUNCTION(BlueprintCallable, Category = "Interchange | Node | Texture")
+	bool GetCustomAddressZ(uint8& AttributeValue) const
+	{
+		IMPLEMENT_NODE_ATTRIBUTE_GETTER(AddressZ, uint8);
+	}
+
+	UFUNCTION(BlueprintCallable, Category = "Interchange | Node | Texture")
+	bool SetCustomAddressZ(const uint8 AttributeValue, bool bAddApplyDelegate = true)
+	{
+#if WITH_EDITORONLY_DATA
+		IMPLEMENT_NODE_ATTRIBUTE_SETTER(UInterchangeTexture2DArrayFactoryNode, AddressZ, uint8, UTexture2DArray)
+#else
+		IMPLEMENT_NODE_ATTRIBUTE_SETTER_NODELEGATE(AddressZ, uint8)
+#endif
+	}
+
+
+private:
+#if WITH_ENGINE
+	virtual void FillAssetClassFromClassName(const FString& ClassName) override
+	{
+		if (ClassName.Equals(UTexture2DArray::StaticClass()->GetName()))
+		{
+			AssetClass = UTexture2DArray::StaticClass();
+			bIsTextureNodeClassInitialized = true;
+		}
+	}
+#endif
+
+	// Addressing
+	IMPLEMENT_NODE_ATTRIBUTE_KEY(AddressX);
+	IMPLEMENT_NODE_ATTRIBUTE_KEY(AddressY);
+	IMPLEMENT_NODE_ATTRIBUTE_KEY(AddressZ);
+
+#if WITH_EDITORONLY_DATA
+	IMPLEMENT_NODE_ATTRIBUTE_APPLY_UOBJECT(AddressX, uint8, UTexture2DArray, TEnumAsByte<TextureAddress>);
+	IMPLEMENT_NODE_ATTRIBUTE_APPLY_UOBJECT(AddressY, uint8, UTexture2DArray, TEnumAsByte<TextureAddress>);
+	IMPLEMENT_NODE_ATTRIBUTE_APPLY_UOBJECT(AddressZ, uint8, UTexture2DArray, TEnumAsByte<TextureAddress>);
+#endif
+};
\ No newline at end of file
diff --git a/Engine/Plugins/Runtime/Interchange/Source/Nodes/Public/InterchangeTexture2DArrayNode.h b/Engine/Plugins/Runtime/Interchange/Source/Nodes/Public/InterchangeTexture2DArrayNode.h
new file mode 100644
--- /dev/null
+++ b/Engine/Plugins/Runtime/Interchange/Source/Nodes/Public/InterchangeTexture2DArrayNode.h
@@ -0,0 +1,28 @@
+// Copyright Epic Games, Inc. All Rights Reserved.
+
+#pragma once
+
+#include "CoreMinimal.h"
+#include "InterchangeTextureNode.h"
+
+#include "InterchangeTexture2DArrayNode.generated.h"
+
+
+
+
+UCLASS(BlueprintType)
+class INTERCHANGENODES_API UInterchangeTexture2DArrayNode : public UInterchangeTextureNode
+{
+	GENERATED_BODY()
+
+public:
+
+	/**
+	 * Return the node type name of the class, we use this when reporting error
+	 */
+	virtual FString GetTypeName() const override
+	{
+		const FString TypeName = TEXT("Texture2DArrayNode");
+		return TypeName;
+	}
+};
\ No newline at end of file
diff --git a/Engine/Plugins/Runtime/Interchange/Source/Nodes/Public/InterchangeTexture2DNode.h b/Engine/Plugins/Runtime/Interchange/Source/Nodes/Public/InterchangeTexture2DNode.h
new file mode 100644
--- /dev/null
+++ b/Engine/Plugins/Runtime/Interchange/Source/Nodes/Public/InterchangeTexture2DNode.h
@@ -0,0 +1,91 @@
+// Copyright Epic Games, Inc. All Rights Reserved.
+
+#pragma once
+
+#include "CoreMinimal.h"
+#include "InterchangeTextureNode.h"
+
+#include "InterchangeTexture2DNode.generated.h"
+
+//Interchange namespace
+namespace UE
+{
+	namespace Interchange
+	{
+		struct FTexture2DNodeStaticData : public FBaseNodeStaticData
+		{
+			static const FString& GetBaseSourceBlocksKey()
+			{
+				static FString StringKey(TEXT("SourceBlocks"));
+				return StringKey;
+			}
+		};
+	}//ns Interchange
+}//ns UE
+
+UCLASS(BlueprintType)
+class INTERCHANGENODES_API UInterchangeTexture2DNode : public UInterchangeTextureNode
+{
+	GENERATED_BODY()
+
+public:
+
+	virtual void PostInitProperties()
+	{
+		Super::PostInitProperties();
+		SourceBlocks.Initialize(Attributes.ToSharedRef(), UE::Interchange::FTexture2DNodeStaticData::GetBaseSourceBlocksKey());
+	}
+
+	/**
+	 * Return the node type name of the class, we use this when reporting error
+	 */
+	virtual FString GetTypeName() const override
+	{
+		const FString TypeName = TEXT("Texture2DNode");
+		return TypeName;
+	}
+
+	//////////////////////////////////////////////////////////////////////////
+	// UDIMs begin here
+	// UDIM base texture use a different model for the source data
+
+	/**
+	 * Get the source blocks for the texture
+	 * If the map is empty then the texture will be simply be imported as normal texture using the payload key
+	 */
+	UFUNCTION(BlueprintCallable, Category = "Interchange | Node | Texture | UDIMs")
+	TMap<int32, FString> GetSourceBlocks() const
+	{
+		return SourceBlocks.ToMap();
+	}
+
+	/**
+	 * Set the source blocks
+	 * Using this will force the texture factory to consider this texture as UDIM.
+	 * @param InSourceBlocks The blocks and their source image that compose the whole texture.
+	 * The textures must be of the same format and use the same pixel format
+	 * The first block inserted in the map be used to determine the accepted texture format and pixel format
+	 */
+	UFUNCTION(BlueprintCallable, Category = "Interchange | Node | Texture | UDIMs")
+	void SetSourceBlocks(const TMap<int32, FString>& InSourceBlocks)
+	{
+		SourceBlocks = InSourceBlocks;
+	}
+
+	/**
+	 * Remove the source block data from texture node
+	 * This is a easy way to stop a texture from being imported as a UDIM
+	 */
+	UFUNCTION(BlueprintCallable, Category = "Interchange | Node | Texture | UDIMs")
+	void ClearSourceBlocksData()
+	{
+		SourceBlocks.Empty();
+	}
+
+	// UDIMs ends here
+	//////////////////////////////////////////////////////////////////////////
+
+protected:
+
+	UE::Interchange::TMapAttributeHelper<int32, FString> SourceBlocks;
+};
diff --git a/Engine/Plugins/Runtime/Interchange/Source/Nodes/Public/InterchangeTextureCubeNode.h b/Engine/Plugins/Runtime/Interchange/Source/Nodes/Public/InterchangeTextureCubeNode.h
--- a/Engine/Plugins/Runtime/Interchange/Source/Nodes/Public/InterchangeTextureCubeNode.h
+++ b/Engine/Plugins/Runtime/Interchange/Source/Nodes/Public/InterchangeTextureCubeNode.h
@@ -3,29 +3,13 @@
 #pragma once
 
 #include "CoreMinimal.h"
-#include "Nodes/InterchangeBaseNode.h"
+#include "InterchangeTextureNode.h"
 
 #include "InterchangeTextureCubeNode.generated.h"
 
-namespace UE
-{
-	namespace Interchange
-	{
-
-		struct FTextureCubeNodeStaticData : public FBaseNodeStaticData
-		{
-			static const FAttributeKey& PayloadSourceFileKey()
-			{
-				static FAttributeKey AttributeKey(TEXT("__PayloadSourceFile__"));
-				return AttributeKey;
-			}
-		};
-	}
-}
-
 
 UCLASS(BlueprintType)
-class INTERCHANGENODES_API UInterchangeTextureCubeNode : public UInterchangeBaseNode
+class INTERCHANGENODES_API UInterchangeTextureCubeNode : public UInterchangeTextureNode
 {
 	GENERATED_BODY()
 
@@ -39,51 +23,4 @@ public:
 		const FString TypeName = TEXT("TextureCubeNode");
 		return TypeName;
 	}
-
-	virtual FString GetKeyDisplayName(const UE::Interchange::FAttributeKey& NodeAttributeKey) const override
-	{
-		FString KeyDisplayName = NodeAttributeKey.ToString();
-		if (NodeAttributeKey == UE::Interchange::FTextureCubeNodeStaticData::PayloadSourceFileKey())
-		{
-			return KeyDisplayName = TEXT("Payload Source Key");
-		}
-		return Super::GetKeyDisplayName(NodeAttributeKey);
-	}
-
-	virtual FGuid GetHash() const override
-	{
-		return Attributes->GetStorageHash();
-	}
-
-	/** Texture node Interface Begin */
-	virtual const TOptional<FString> GetPayLoadKey() const
-	{
-		if (!Attributes->ContainAttribute(UE::Interchange::FTextureCubeNodeStaticData::PayloadSourceFileKey()))
-		{
-			return TOptional<FString>();
-		}
-		UE::Interchange::FAttributeStorage::TAttributeHandle<FString> AttributeHandle = Attributes->GetAttributeHandle<FString>(UE::Interchange::FTextureCubeNodeStaticData::PayloadSourceFileKey());
-		if (!AttributeHandle.IsValid())
-		{
-			return TOptional<FString>();
-		}
-		FString PayloadKey;
-		UE::Interchange::EAttributeStorageResult Result = AttributeHandle.Get(PayloadKey);
-		if (!IsAttributeStorageResultSuccess(Result))
-		{
-			LogAttributeStorageErrors(Result, TEXT("UInterchangeTextureNode.GetPayLoadKey"), UE::Interchange::FTextureCubeNodeStaticData::PayloadSourceFileKey());
-			return TOptional<FString>();
-		}
-		return TOptional<FString>(PayloadKey);
-	}
-
-	UFUNCTION(BlueprintCallable, Category = "Interchange | Node | Texture")
-	virtual void SetPayLoadKey(const FString& PayloadKey)
-	{
-		UE::Interchange::EAttributeStorageResult Result = Attributes->RegisterAttribute(UE::Interchange::FTextureCubeNodeStaticData::PayloadSourceFileKey(), PayloadKey);
-		if (!IsAttributeStorageResultSuccess(Result))
-		{
-			LogAttributeStorageErrors(Result, TEXT("UInterchangeTextureNode.SetPayLoadKey"), UE::Interchange::FTextureCubeNodeStaticData::PayloadSourceFileKey());
-		}
-	}
 };
diff --git a/Engine/Plugins/Runtime/Interchange/Source/Nodes/Public/InterchangeTextureFactoryNode.h b/Engine/Plugins/Runtime/Interchange/Source/Nodes/Public/InterchangeTextureFactoryNode.h
--- a/Engine/Plugins/Runtime/Interchange/Source/Nodes/Public/InterchangeTextureFactoryNode.h
+++ b/Engine/Plugins/Runtime/Interchange/Source/Nodes/Public/InterchangeTextureFactoryNode.h
@@ -561,26 +561,26 @@ public:
 
 	/** Return false if the Attribute was not set previously.*/
 	UFUNCTION(BlueprintCallable, Category = "Interchange | Node | Texture")
-	bool GetCustomAddressX(uint8& AttributeValue) const
+	virtual bool GetCustomAddressX(uint8& AttributeValue) const
 	{
 		IMPLEMENT_NODE_ATTRIBUTE_GETTER(AddressX, uint8);
 	}
 
 	UFUNCTION(BlueprintCallable, Category = "Interchange | Node | Texture")
-	bool SetCustomAddressX(const uint8& AttributeValue, bool bAddApplyDelegate = true)
+	virtual bool SetCustomAddressX(const uint8 AttributeValue, bool bAddApplyDelegate = true)
 	{
 		IMPLEMENT_NODE_ATTRIBUTE_SETTER(UInterchangeTextureFactoryNode, AddressX, uint8, UTexture2D)
 	}
 
 	/** Return false if the Attribute was not set previously.*/
 	UFUNCTION(BlueprintCallable, Category = "Interchange | Node | Texture")
-	bool GetCustomAddressY(uint8& AttributeValue) const
+	virtual bool GetCustomAddressY(uint8& AttributeValue) const
 	{
 		IMPLEMENT_NODE_ATTRIBUTE_GETTER(AddressY, uint8);
 	}
 
 	UFUNCTION(BlueprintCallable, Category = "Interchange | Node | Texture")
-	bool SetCustomAddressY(const uint8& AttributeValue, bool bAddApplyDelegate = true)
+	virtual bool SetCustomAddressY(const uint8 AttributeValue, bool bAddApplyDelegate = true)
 	{
 		IMPLEMENT_NODE_ATTRIBUTE_SETTER(UInterchangeTextureFactoryNode, AddressY, uint8, UTexture2D)
 	}
diff --git a/Engine/Plugins/Runtime/Interchange/Source/Nodes/Public/InterchangeTextureNode.h b/Engine/Plugins/Runtime/Interchange/Source/Nodes/Public/InterchangeTextureNode.h
--- a/Engine/Plugins/Runtime/Interchange/Source/Nodes/Public/InterchangeTextureNode.h
+++ b/Engine/Plugins/Runtime/Interchange/Source/Nodes/Public/InterchangeTextureNode.h
@@ -12,7 +12,6 @@ namespace UE
 {
 	namespace Interchange
 	{
-
 		struct FTextureNodeStaticData : public FBaseNodeStaticData
 		{
 			static const FAttributeKey& PayloadSourceFileKey()
@@ -20,33 +19,16 @@ namespace UE
 				static FAttributeKey AttributeKey(TEXT("__PayloadSourceFile__"));
 				return AttributeKey;
 			}
-
-			static const FString& GetBaseSourceBlocksKey()
-			{
-				static FString StringKey(TEXT("SourceBlocks"));
-				return StringKey;
-			}
 		};
 	}//ns Interchange
 }//ns UE
 
-UCLASS(BlueprintType)
+UCLASS(BlueprintType, Abstract)
 class INTERCHANGENODES_API UInterchangeTextureNode : public UInterchangeBaseNode
 {
 	GENERATED_BODY()
 
 public:
-	UInterchangeTextureNode()
-	:UInterchangeBaseNode()
-	{
-	}
-
-	virtual void PostInitProperties()
-	{
-		Super::PostInitProperties();
-		SourceBlocks.Initialize(Attributes.ToSharedRef(), UE::Interchange::FTextureNodeStaticData::GetBaseSourceBlocksKey());
-	}
-
 	/**
 	 * Return the node type name of the class, we use this when reporting error
 	 */
@@ -102,48 +84,4 @@ public:
 			LogAttributeStorageErrors(Result, TEXT("UInterchangeTextureNode.SetPayLoadKey"), UE::Interchange::FTextureNodeStaticData::PayloadSourceFileKey());
 		}
 	}
-
-	//////////////////////////////////////////////////////////////////////////
-	// UDIMs begin here
-	// UDIM base texture use a different model for the source data
-
-	/**
-	 * Get the source blocks for the texture
-	 * If the map is empty then the texture will be simply be imported as normal texture using the payload key
-	 */
-	UFUNCTION(BlueprintCallable, Category = "Interchange | Node | Texture | UDIMs")
-	TMap<int32, FString> GetSourceBlocks() const
-	{
-		return SourceBlocks.ToMap();
-	}
-
-	/**
-	 * Set the source blocks
-	 * Using this will force the texture factory to consider this texture as UDIM.
-	 * @param InSourceBlocks The blocks and their source image that compose the whole texture.
-	 * The textures must be of the same format and use the same pixel format
-	 * The first block inserted in the map be used to determine the accepted texture format and pixel format
-	 */
-	UFUNCTION(BlueprintCallable, Category = "Interchange | Node | Texture | UDIMs")
-	void SetSourceBlocks(const TMap<int32, FString>& InSourceBlocks)
-	{
-		SourceBlocks = InSourceBlocks;
-	}
-
-	/**
-	 * Remove the source block data from texture node
-	 * This is a easy way to stop a texture from being imported as a UDIM
-	 */
-	UFUNCTION(BlueprintCallable, Category = "Interchange | Node | Texture | UDIMs")
-	void ClearSourceBlocksData()
-	{
-		SourceBlocks.Empty();
-	}
-
-	// UDIMs ends here
-	//////////////////////////////////////////////////////////////////////////
-
-protected:
-
-	UE::Interchange::TMapAttributeHelper<int32, FString> SourceBlocks;
 };
diff --git a/Engine/Plugins/Runtime/Interchange/Source/Parsers/Fbx/Private/FbxMaterial.cpp b/Engine/Plugins/Runtime/Interchange/Source/Parsers/Fbx/Private/FbxMaterial.cpp
--- a/Engine/Plugins/Runtime/Interchange/Source/Parsers/Fbx/Private/FbxMaterial.cpp
+++ b/Engine/Plugins/Runtime/Interchange/Source/Parsers/Fbx/Private/FbxMaterial.cpp
@@ -8,6 +8,7 @@
 #include "FbxInclude.h"
 #include "InterchangeMaterialNode.h"
 #include "InterchangeSceneNode.h"
+#include "InterchangeTexture2DNode.h"
 #include "InterchangeTextureNode.h"
 #include "Misc/Paths.h"
 #include "Nodes/InterchangeBaseNodeContainer.h"
@@ -33,10 +34,10 @@ namespace UE
 				return MaterialNode;
 			}
 
-			UInterchangeTextureNode* FFbxMaterial::CreateTexture2DNode(UInterchangeBaseNodeContainer& NodeContainer, const FString& NodeUid, const FString& TextureFilePath, TArray<FString>& JSonErrorMessages)
+			UInterchangeTexture2DNode* FFbxMaterial::CreateTexture2DNode(UInterchangeBaseNodeContainer& NodeContainer, const FString& NodeUid, const FString& TextureFilePath, TArray<FString>& JSonErrorMessages)
 			{
 				FString DisplayLabel = FPaths::GetBaseFilename(TextureFilePath);
-				UInterchangeTextureNode* TextureNode = NewObject<UInterchangeTextureNode>(&NodeContainer, NAME_None);
+				UInterchangeTexture2DNode* TextureNode = NewObject<UInterchangeTexture2DNode>(&NodeContainer, NAME_None);
 				if (!ensure(TextureNode))
 				{
 					JSonErrorMessages.Add(TEXT("{\"Msg\" : {\"Type\" : \"Error\",\n\"Msg\" : \"Cannot allocate a node when importing fbx\"}}"));
@@ -95,7 +96,7 @@ namespace UE
 									//Create a texture node and make it child of the material node
 									TArray<FString> JsonErrorMessage;
 									FString NodeUid = TEXT("\\Texture\\") + TextureFilename;
-									UInterchangeTextureNode* TextureNode = Cast<UInterchangeTextureNode>(NodeContainer.GetNode(NodeUid));
+									UInterchangeTexture2DNode* TextureNode = Cast<UInterchangeTexture2DNode>(NodeContainer.GetNode(NodeUid));
 									if (!TextureNode)
 									{
 										TextureNode = CreateTexture2DNode(NodeContainer, NodeUid, TextureFilename, JsonErrorMessage);
@@ -187,7 +188,7 @@ namespace UE
 					//Create a texture node and make it child of the material node
 					TArray<FString> JsonErrorMessage;
 					FString NodeUid = TEXT("\\Texture\\") + TextureFilename;
-					UInterchangeTextureNode* TextureNode = Cast<UInterchangeTextureNode>(NodeContainer.GetNode(NodeUid));
+					UInterchangeTexture2DNode* TextureNode = Cast<UInterchangeTexture2DNode>(NodeContainer.GetNode(NodeUid));
 					if (!TextureNode)
 					{
 						TextureNode = CreateTexture2DNode(NodeContainer, NodeUid, TextureFilename, JsonErrorMessage);
diff --git a/Engine/Plugins/Runtime/Interchange/Source/Parsers/Fbx/Private/FbxMaterial.h b/Engine/Plugins/Runtime/Interchange/Source/Parsers/Fbx/Private/FbxMaterial.h
--- a/Engine/Plugins/Runtime/Interchange/Source/Parsers/Fbx/Private/FbxMaterial.h
+++ b/Engine/Plugins/Runtime/Interchange/Source/Parsers/Fbx/Private/FbxMaterial.h
@@ -9,7 +9,7 @@
 class UInterchangeBaseNode;
 class UInterchangeBaseNodeContainer;
 class UInterchangeMaterialNode;
-class UInterchangeTextureNode;
+class UInterchangeTexture2DNode;
 class UInterchangeSceneNode;
 
 namespace UE
@@ -44,7 +44,7 @@ namespace UE
 				static void AddAllNodeMaterials(UInterchangeSceneNode* SceneNode, FbxNode* ParentFbxNode, UInterchangeBaseNodeContainer& NodeContainer, TArray<FString>& JSonErrorMessages);
 			protected:
 				static UInterchangeMaterialNode* CreateMaterialNode(UInterchangeBaseNodeContainer& NodeContainer, const FString& NodeUID, const FString& NodeName, TArray<FString>& JSonErrorMessages);
-				static UInterchangeTextureNode* CreateTexture2DNode(UInterchangeBaseNodeContainer& NodeContainer, const FString& NodeUID, const FString& TextureFilePath, TArray<FString>& JSonErrorMessages);
+				static UInterchangeTexture2DNode* CreateTexture2DNode(UInterchangeBaseNodeContainer& NodeContainer, const FString& NodeUID, const FString& TextureFilePath, TArray<FString>& JSonErrorMessages);
 
 			private:
 				static UInterchangeMaterialNode* AddNodeMaterial(FbxSurfaceMaterial* SurfaceMaterial, UInterchangeBaseNodeContainer& NodeContainer, TArray<FString>& JSonErrorMessages);
diff --git a/Engine/Plugins/Runtime/Interchange/Source/Pipelines/Private/InterchangeGenericAssetsPipeline.cpp b/Engine/Plugins/Runtime/Interchange/Source/Pipelines/Private/InterchangeGenericAssetsPipeline.cpp
--- a/Engine/Plugins/Runtime/Interchange/Source/Pipelines/Private/InterchangeGenericAssetsPipeline.cpp
+++ b/Engine/Plugins/Runtime/Interchange/Source/Pipelines/Private/InterchangeGenericAssetsPipeline.cpp
@@ -14,6 +14,9 @@
 #include "InterchangeSkeletalMeshLodDataNode.h"
 #include "InterchangeSkeletonFactoryNode.h"
 #include "InterchangeSourceData.h"
+#include "InterchangeTexture2DArrayFactoryNode.h"
+#include "InterchangeTexture2DArrayNode.h"
+#include "InterchangeTexture2DNode.h"
 #include "InterchangeTextureCubeFactoryNode.h"
 #include "InterchangeTextureCubeNode.h"
 #include "InterchangeTextureFactoryNode.h"
@@ -21,9 +24,54 @@
 #include "Misc/Paths.h"
 #include "Nodes/InterchangeBaseNode.h"
 #include "Nodes/InterchangeBaseNodeContainer.h"
+#include "Templates/SubclassOf.h"
 #include "UObject/Object.h"
 #include "UObject/ObjectMacros.h"
 
+namespace UE::Interchange::Private
+{
+	UClass* GetDefaultFactoryClassFromTextureNodeClass(UClass* NodeClass)
+	{
+		if (UInterchangeTexture2DNode::StaticClass() == NodeClass)
+		{
+			return UInterchangeTextureFactoryNode::StaticClass();
+		}
+
+		if (UInterchangeTextureCubeNode::StaticClass() == NodeClass)
+		{
+			return UInterchangeTextureCubeFactoryNode::StaticClass();
+		}
+
+		if (UInterchangeTexture2DArrayNode::StaticClass() == NodeClass)
+		{
+			return UInterchangeTexture2DArrayFactoryNode::StaticClass();
+		}
+
+		return nullptr;
+	}
+
+	UClass* GetDefaultAssetClassFromFactoryClass(UClass* NodeClass)
+	{
+		if (UInterchangeTextureFactoryNode::StaticClass() == NodeClass)
+		{
+			return UTexture2D::StaticClass();
+		}
+
+		if (UInterchangeTextureCubeFactoryNode::StaticClass() == NodeClass)
+		{
+			return UTextureCube::StaticClass();
+		}
+
+		if (UInterchangeTexture2DArrayFactoryNode::StaticClass() == NodeClass)
+		{
+			return UTexture2DArray::StaticClass();
+		}
+
+		return nullptr;
+	}
+}
+
+
 bool UInterchangeGenericAssetsPipeline::ExecutePreImportPipeline(UInterchangeBaseNodeContainer* InBaseNodeContainer, const TArray<UInterchangeSourceData*>& InSourceDatas)
 {
 	if (!InBaseNodeContainer)
@@ -71,11 +119,6 @@ bool UInterchangeGenericAssetsPipeline::ExecutePreImportPipeline(UInterchangeBas
 					{
 						StaticMeshNodes.Add(MeshNode);
 					}
-
-				}
-				else if (UInterchangeTextureCubeNode* TextuteCubeNode = Cast<UInterchangeTextureCubeNode>(Node))
-				{
-					TextureCubeNodes.Add(TextuteCubeNode);
 				}
 			}
 			break;
@@ -95,16 +138,10 @@ bool UInterchangeGenericAssetsPipeline::ExecutePreImportPipeline(UInterchangeBas
 		//import textures
 		for (const UInterchangeTextureNode* TextureNode : TextureNodes)
 		{
-			CreateTexture2DFactoryNode(TextureNode);
+			CreateTextureFactoryNode(TextureNode, UE::Interchange::Private::GetDefaultFactoryClassFromTextureNodeClass(TextureNode->GetClass()));
 		}
 	}
 
-	for (const UInterchangeTextureCubeNode* TextureCubeNode : TextureCubeNodes)
-	{
-		CreateTextureCubeFactoryNode(TextureCubeNode);
-	}
-
-
 	if (bImportMaterials)
 	{
 		//import materials
@@ -306,7 +343,7 @@ bool UInterchangeGenericAssetsPipeline::ExecutePreImportPipeline(UInterchangeBas
 	return true;
 }
 
-UInterchangeTextureFactoryNode* UInterchangeGenericAssetsPipeline::CreateTexture2DFactoryNode(const UInterchangeTextureNode* TextureNode)
+UInterchangeTextureFactoryNode* UInterchangeGenericAssetsPipeline::CreateTextureFactoryNode(const UInterchangeTextureNode* TextureNode, const TSubclassOf<UInterchangeTextureFactoryNode>& FactorySubclass)
 {
 	FString DisplayLabel = TextureNode->GetDisplayLabel();
 	FString NodeUid = UInterchangeTextureFactoryNode::GetTextureFactoryNodeUidFromTextureNodeUid(TextureNode->GetUniqueID());
@@ -321,48 +358,32 @@ UInterchangeTextureFactoryNode* UInterchangeGenericAssetsPipeline::CreateTexture
 	}
 	else
 	{
-		TextureFactoryNode = NewObject<UInterchangeTextureFactoryNode>(BaseNodeContainer, NAME_None);
-		if (!ensure(TextureFactoryNode))
+		UClass* FactoryClass = FactorySubclass.Get();
+		if (!ensure(FactoryClass))
 		{
+			// Log an error
 			return nullptr;
 		}
-		//Creating a UTexture2D
-		TextureFactoryNode->InitializeTextureNode(NodeUid, DisplayLabel, UTexture2D::StaticClass()->GetName(), TextureNode->GetDisplayLabel());
-		TextureFactoryNode->SetCustomTranslatedTextureNodeUid(TextureNode->GetUniqueID());
-		BaseNodeContainer->AddNode(TextureFactoryNode);
-		TextureFactoryNodes.Add(TextureFactoryNode);
-	}
-	return TextureFactoryNode;
-}
 
-UInterchangeTextureCubeFactoryNode* UInterchangeGenericAssetsPipeline::CreateTextureCubeFactoryNode(const UInterchangeTextureCubeNode* TextureCubeNode)
-{
-	FString DisplayLabel = TextureCubeNode->GetDisplayLabel();
-	FString NodeUid = UInterchangeTextureFactoryNode::GetTextureFactoryNodeUidFromTextureNodeUid(TextureCubeNode->GetUniqueID());
-	UInterchangeTextureCubeFactoryNode* TextureCubeFactoryNode = nullptr;
-	if (BaseNodeContainer->IsNodeUidValid(NodeUid))
-	{
-		TextureCubeFactoryNode = Cast<UInterchangeTextureCubeFactoryNode>(BaseNodeContainer->GetNode(NodeUid));
-		if (!ensure(TextureCubeFactoryNode))
+		UClass* TextureClass = UE::Interchange::Private::GetDefaultAssetClassFromFactoryClass(FactoryClass);
+		if (!ensure(TextureClass))
 		{
-			//Log an error
+			// Log an error
+			return nullptr;
 		}
-	}
-	else
-	{
-		TextureCubeFactoryNode = NewObject<UInterchangeTextureCubeFactoryNode>(BaseNodeContainer, NAME_None);
-		if (!ensure(TextureCubeFactoryNode))
+
+		TextureFactoryNode = NewObject<UInterchangeTextureFactoryNode>(BaseNodeContainer, FactoryClass);
+		if (!ensure(TextureFactoryNode))
 		{
 			return nullptr;
 		}
-
-		//Creating a UTextureCube
-		TextureCubeFactoryNode->InitializeTextureNode(NodeUid, DisplayLabel, UTextureCube::StaticClass()->GetName(), TextureCubeNode->GetDisplayLabel());
-		TextureCubeFactoryNode->SetCustomTranslatedTextureNodeUid(TextureCubeNode->GetUniqueID());
-		BaseNodeContainer->AddNode(TextureCubeFactoryNode);
-		TextureCubeFactoryNodes.Add(TextureCubeFactoryNode);
+		//Creating a UTexture2D
+		TextureFactoryNode->InitializeTextureNode(NodeUid, DisplayLabel, TextureClass->GetName(), TextureNode->GetDisplayLabel());
+		TextureFactoryNode->SetCustomTranslatedTextureNodeUid(TextureNode->GetUniqueID());
+		BaseNodeContainer->AddNode(TextureFactoryNode);
+		TextureFactoryNodes.Add(TextureFactoryNode);
 	}
-	return TextureCubeFactoryNode;
+	return TextureFactoryNode;
 }
 
 UInterchangeMaterialFactoryNode* UInterchangeGenericAssetsPipeline::CreateMaterialFactoryNode(const UInterchangeMaterialNode* MaterialNode)
diff --git a/Engine/Plugins/Runtime/Interchange/Source/Pipelines/Public/InterchangeGenericAssetsPipeline.h b/Engine/Plugins/Runtime/Interchange/Source/Pipelines/Public/InterchangeGenericAssetsPipeline.h
--- a/Engine/Plugins/Runtime/Interchange/Source/Pipelines/Public/InterchangeGenericAssetsPipeline.h
+++ b/Engine/Plugins/Runtime/Interchange/Source/Pipelines/Public/InterchangeGenericAssetsPipeline.h
@@ -19,10 +19,13 @@ class UInterchangeSceneNode;
 class UInterchangeSkeletalMeshFactoryNode;
 class UInterchangeSkeletalMeshLodDataNode;
 class UInterchangeSkeletonFactoryNode;
+class UInterchangeTexture2DArrayFactoryNode;
+class UInterchangeTexture2DArrayNode;
 class UInterchangeTextureCubeFactoryNode;
 class UInterchangeTextureCubeNode;
 class UInterchangeTextureFactoryNode;
 class UInterchangeTextureNode;
+template<class T> class TSubclassOf;
 
 #define COMMON_CATEGORY "Common"
 #define COMMON_MESHES_CATEGORY "Common Meshes"
@@ -161,17 +164,7 @@ private:
 	/** Texture factory assets nodes */
 	TArray<UInterchangeTextureFactoryNode*> TextureFactoryNodes;
 
-
-	/** Texture cube translated assets nodes*/
-	TArray<UInterchangeTextureCubeNode*> TextureCubeNodes;
-
-	/** Texture cube factory translated assets nodes */
-	TArray<UInterchangeTextureCubeFactoryNode*> TextureCubeFactoryNodes;
-
-
-	UInterchangeTextureFactoryNode* CreateTexture2DFactoryNode(const UInterchangeTextureNode* TextureNode);
-
-	UInterchangeTextureCubeFactoryNode* CreateTextureCubeFactoryNode(const UInterchangeTextureCubeNode* TextureCubeNode);
+	UInterchangeTextureFactoryNode* CreateTextureFactoryNode(const UInterchangeTextureNode* TextureNode, const TSubclassOf<UInterchangeTextureFactoryNode>& FactorySubclass);
 	
 	/** Material translated assets nodes */
 	TArray<UInterchangeMaterialNode*> MaterialNodes;
