Submitted a mising File to my previous CL

Added texture arrays support to interchange.

Added support to import texture array from DDS File.

The texture 2d array factory is still missing some feature such as creating a empty texture and setting the SourceTextures.
Those will be done in another pass for some other jira.

diff --git a/Engine/Plugins/Runtime/Interchange/Source/Import/Private/Texture/InterchangeDDSTranslator.cpp b/Engine/Plugins/Runtime/Interchange/Source/Import/Private/Texture/InterchangeDDSTranslator.cpp
--- a/Engine/Plugins/Runtime/Interchange/Source/Import/Private/Texture/InterchangeDDSTranslator.cpp
+++ b/Engine/Plugins/Runtime/Interchange/Source/Import/Private/Texture/InterchangeDDSTranslator.cpp
@@ -76,7 +76,7 @@ bool UInterchangeDDSTranslator::CanImportSourceData(const UInterchangeSourceData
 	}
 
 	FDDSLoadHelper  DDSLoadHelper(HeaderDataBuffer.GetData(), HeaderDataBuffer.Num());
-	return DDSLoadHelper.IsValid2DTexture() || DDSLoadHelper.IsValidCubemapTexture();
+	return DDSLoadHelper.IsValid2DTexture() || DDSLoadHelper.IsValidCubemapTexture() || DDSLoadHelper.IsValidArrayTexture();
 }
 
 bool UInterchangeDDSTranslator::Translate(const UInterchangeSourceData* SourceData, UInterchangeBaseNodeContainer& BaseNodeContainer) const
@@ -241,7 +241,7 @@ TOptional<UE::Interchange::FImportSlicedImage> UInterchangeDDSTranslator::GetSli
 	//
 	FDDSLoadHelper DDSLoadHelper(DDSSourceData.GetData(), DDSSourceData.Num());
 
-	if (!DDSLoadHelper.IsValidCubemapTexture())
+	if (!DDSLoadHelper.IsValidCubemapTexture() && !DDSLoadHelper.IsValidArrayTexture() )
 	{
 		UE_LOG(LogInterchangeImport, Error, TEXT("Failed to import DDS, unsupported format. [%s]"), *Filename);
 		return TOptional<UE::Interchange::FImportSlicedImage>();
