Interchange add support for IES Files and TextureLightProfile

Made the texture factory more fexible so that the textures can accept a payload that aren't specificly made for their type asset if desired.

diff --git a/Engine/Plugins/Runtime/Interchange/Source/Import/InterchangeImport.Build.cs b/Engine/Plugins/Runtime/Interchange/Source/Import/InterchangeImport.Build.cs
--- a/Engine/Plugins/Runtime/Interchange/Source/Import/InterchangeImport.Build.cs
+++ b/Engine/Plugins/Runtime/Interchange/Source/Import/InterchangeImport.Build.cs
@@ -24,6 +24,7 @@ namespace UnrealBuildTool.Rules
 			PrivateDependencyModuleNames.AddRange(
 				new string[]
 				{
+					"IESFile",
 					"ImageWrapper",
 					"InterchangeDispatcher",
 					"Json",
diff --git a/Engine/Plugins/Runtime/Interchange/Source/Import/Private/InterchangeImportModule.cpp b/Engine/Plugins/Runtime/Interchange/Source/Import/Private/InterchangeImportModule.cpp
--- a/Engine/Plugins/Runtime/Interchange/Source/Import/Private/InterchangeImportModule.cpp
+++ b/Engine/Plugins/Runtime/Interchange/Source/Import/Private/InterchangeImportModule.cpp
@@ -15,6 +15,7 @@
 #include "Texture/InterchangeDDSTranslator.h"
 #include "Texture/InterchangeEXRTranslator.h"
 #include "Texture/InterchangeHDRTranslator.h"
+#include "Texture/InterchangeIESTranslator.h"
 #include "Texture/InterchangeJPGTranslator.h"
 #include "Texture/InterchangePCXTranslator.h"
 #include "Texture/InterchangePNGTranslator.h"
@@ -51,7 +52,6 @@ void FInterchangeImportModule::StartupModule()
 		// UDIM must be registered before the other texture translators
 		InterchangeManager.RegisterTranslator(UInterchangeUDIMTranslator::StaticClass());
 
-
 		InterchangeManager.RegisterTranslator(UInterchangeBMPTranslator::StaticClass());
 		InterchangeManager.RegisterTranslator(UInterchangeDDSTranslator::StaticClass());
 		InterchangeManager.RegisterTranslator(UInterchangeEXRTranslator::StaticClass());
@@ -61,6 +61,7 @@ void FInterchangeImportModule::StartupModule()
 		InterchangeManager.RegisterTranslator(UInterchangePSDTranslator::StaticClass());
 		InterchangeManager.RegisterTranslator(UInterchangeTGATranslator::StaticClass());
 		InterchangeManager.RegisterTranslator(UInterchangeHDRTranslator::StaticClass());
+		InterchangeManager.RegisterTranslator(UInterchangeIESTranslator::StaticClass());
 
 		//Register the factories
 		InterchangeManager.RegisterFactory(UInterchangeTextureFactory::StaticClass());
diff --git a/Engine/Plugins/Runtime/Interchange/Source/Import/Private/Texture/InterchangeIESTranslator.cpp b/Engine/Plugins/Runtime/Interchange/Source/Import/Private/Texture/InterchangeIESTranslator.cpp
new file mode 100644
--- /dev/null
+++ b/Engine/Plugins/Runtime/Interchange/Source/Import/Private/Texture/InterchangeIESTranslator.cpp
@@ -0,0 +1,81 @@
+// Copyright Epic Games, Inc. All Rights Reserved. 
+
+#include "Texture/InterchangeIESTranslator.h"
+
+#include "Engine/Texture.h"
+#include "Engine/TextureLightProfile.h"
+#include "IESConverter.h"
+#include "InterchangeImportLog.h"
+#include "Misc/ConfigCacheIni.h"
+#include "Misc/FileHelper.h"
+#include "Misc/Paths.h"
+#include "Texture/TextureTranslatorUtilities.h"
+
+bool UInterchangeIESTranslator::CanImportSourceData(const UInterchangeSourceData* InSourceData) const
+{
+	FString Extension = FPaths::GetExtension(InSourceData->GetFilename());
+	FString IESExtension = (TEXT("ies"));
+	return IESExtension.StartsWith(Extension);
+}
+
+bool UInterchangeIESTranslator::Translate(const UInterchangeSourceData* SourceData, UInterchangeBaseNodeContainer& BaseNodeContainer) const
+{
+	return UE::Interchange::FTextureTranslatorUtilities::GenericTextureLightProfileTranslate(SourceData, BaseNodeContainer);
+}
+
+TOptional<UE::Interchange::FImportLightProfile> UInterchangeIESTranslator::GetLightProfilePayloadData(const UInterchangeSourceData* SourceData, const FString& PayLoadKey) const
+{
+	if (!SourceData)
+	{
+		UE_LOG(LogInterchangeImport, Error, TEXT("Failed to import IES, bad source data."));
+		return {};
+	}
+
+	TArray64<uint8> SourceDataBuffer;
+	FString Filename = SourceData->GetFilename();
+
+	//Make sure the key fit the filename, The key should always be valid
+	if (!Filename.Equals(PayLoadKey))
+	{
+		UE_LOG(LogInterchangeImport, Error, TEXT("Failed to import IES, wrong payload key. [%s]"), *Filename);
+		return {};
+	}
+
+	if (!FPaths::FileExists(Filename))
+	{
+		UE_LOG(LogInterchangeImport, Error, TEXT("Failed to import IES, cannot open file. [%s]"), *Filename);
+		return {};
+	}
+
+	if (!FFileHelper::LoadFileToArray(SourceDataBuffer, *Filename))
+	{
+		UE_LOG(LogInterchangeImport, Error, TEXT("Failed to import IES, cannot load file content into an array. [%s]"), *Filename);
+		return {};
+	}
+
+	const uint8* Buffer = SourceDataBuffer.GetData();
+
+	FIESConverter IESConverter(Buffer, SourceDataBuffer.Num());
+
+	if(IESConverter.IsValid())
+	{
+		UE::Interchange::FImportLightProfile Payload;
+
+		Payload.Init2DWithParams(
+			IESConverter.GetWidth(),
+			IESConverter.GetHeight(),
+			TSF_RGBA16F,
+			false
+			);
+
+		Payload.CompressionSettings = TC_HDR;
+		Payload.Brightness = IESConverter.GetBrightness();
+		Payload.TextureMultiplier = IESConverter.GetMultiplier();
+
+		Payload.RawData = IESConverter.GetRawData();
+
+		return Payload;
+	}
+
+	return {};
+}
diff --git a/Engine/Plugins/Runtime/Interchange/Source/Import/Private/Texture/InterchangeTextureFactory.cpp b/Engine/Plugins/Runtime/Interchange/Source/Import/Private/Texture/InterchangeTextureFactory.cpp
--- a/Engine/Plugins/Runtime/Interchange/Source/Import/Private/Texture/InterchangeTextureFactory.cpp
+++ b/Engine/Plugins/Runtime/Interchange/Source/Import/Private/Texture/InterchangeTextureFactory.cpp
@@ -6,6 +6,7 @@
 #include "Engine/Texture2D.h"
 #include "Engine/Texture2DArray.h"
 #include "Engine/TextureCube.h"
+#include "Engine/TextureLightProfile.h"
 #include "InterchangeAssetImportData.h"
 #include "InterchangeImportCommon.h"
 #include "InterchangeImportLog.h"
@@ -15,12 +16,16 @@
 #include "InterchangeTextureCubeFactoryNode.h"
 #include "InterchangeTextureCubeNode.h"
 #include "InterchangeTextureFactoryNode.h"
+#include "InterchangeTextureLightProfileFactoryNode.h"
+#include "InterchangeTextureLightProfileNode.h"
 #include "InterchangeTextureNode.h"
 #include "InterchangeTranslatorBase.h"
 #include "Misc/TVariant.h"
 #include "Nodes/InterchangeBaseNode.h"
 #include "Texture/InterchangeBlockedTexturePayloadInterface.h"
 #include "Texture/InterchangeSlicedTexturePayloadInterface.h"
+#include "Texture/InterchangeTextureLightProfilePayloadData.h"
+#include "Texture/InterchangeTextureLightProfilePayloadInterface.h"
 #include "Texture/InterchangeTexturePayloadInterface.h"
 #include "TextureCompiler.h"
 #include "UObject/ObjectMacros.h"
@@ -43,26 +48,32 @@ namespace UE::Interchange::Private::InterchangeTextureFactory
 		UClass* TextureCubeFactoryClass = UInterchangeTextureCubeFactoryNode::StaticClass();
 		UClass* TextureFactoryClass = UInterchangeTextureFactoryNode::StaticClass();
 		UClass* Texture2DArrayFactoryClass = UInterchangeTexture2DArrayFactoryNode::StaticClass();
+		UClass* TextureLightProfileFactoryClass = UInterchangeTextureLightProfileFactoryNode::StaticClass();
 		UClass* AssetClass = AssetNode->GetClass();
 #if USTRUCT_FAST_ISCHILDOF_IMPL == USTRUCT_ISCHILDOF_STRUCTARRAY
 		if (AssetClass->IsChildOf(Texture2DArrayFactoryClass))
 		{
 			return Texture2DArrayFactoryClass;
 		}
-		else if (AssetClass->IsChildOf(TextureFactoryClass))
-		{
-			return TextureFactoryClass;
-		}
 		else if (AssetClass->IsChildOf(TextureCubeFactoryClass))
 		{
 			return TextureCubeFactoryClass;
 		}
+		else if (AssetClass->IsChildOf(TextureLightProfileFactoryClass))
+		{
+			return TextureLightProfileFactoryClass;
+		}
+		else if (AssetClass->IsChildOf(TextureFactoryClass))
+		{
+			return TextureFactoryClass;
+		}
 #else
 		while (AssetClass)
 		{
 			if (AssetClass == TextureCubeFactoryClass
 				|| AssetClass == TextureFactoryClass
-				|| AssetClass == Texture2DArrayFactoryClass)
+				|| AssetClass == Texture2DArrayFactoryClass
+				|| AssetClass == TextureLightProfileFactoryClass)
 			{
 				return AssetClass;
 			}
@@ -77,7 +88,8 @@ namespace UE::Interchange::Private::InterchangeTextureFactory
 	using FTextureFactoryNodeVariant = TVariant<FEmptyVariantState
 		, UInterchangeTextureFactoryNode*
 		, UInterchangeTextureCubeFactoryNode*
-		, UInterchangeTexture2DArrayFactoryNode*>; 
+		, UInterchangeTexture2DArrayFactoryNode*
+		,UInterchangeTextureLightProfileFactoryNode* >;
 
 	FTextureFactoryNodeVariant GetAsTextureFactoryNodeVariant(UInterchangeBaseNode* AssetNode, UClass* SupportedFactoryNodeClass)
 	{
@@ -102,6 +114,11 @@ namespace UE::Interchange::Private::InterchangeTextureFactory
 			{
 				return FTextureFactoryNodeVariant(TInPlaceType<UInterchangeTexture2DArrayFactoryNode*>(), static_cast<UInterchangeTexture2DArrayFactoryNode*>(AssetNode));
 			}
+
+			if (SupportedFactoryNodeClass == UInterchangeTextureLightProfileFactoryNode::StaticClass())
+			{
+				return FTextureFactoryNodeVariant(TInPlaceType<UInterchangeTextureLightProfileFactoryNode*>(), static_cast<UInterchangeTextureLightProfileFactoryNode*>(AssetNode));
+			}
 		}
 
 		return {};
@@ -110,7 +127,8 @@ namespace UE::Interchange::Private::InterchangeTextureFactory
 	using FTextureNodeVariant = TVariant<FEmptyVariantState
 		, const UInterchangeTexture2DNode*
 		, const UInterchangeTextureCubeNode*
-		, const UInterchangeTexture2DArrayNode*>; 
+		, const UInterchangeTexture2DArrayNode*
+		, const UInterchangeTextureLightProfileNode* >;
 
 	FTextureNodeVariant GetTextureNodeVariantFromFactoryVariant(const FTextureFactoryNodeVariant& FactoryVariant, const UInterchangeBaseNodeContainer* NodeContainer)
 	{
@@ -118,39 +136,41 @@ namespace UE::Interchange::Private::InterchangeTextureFactory
 
 		if (UInterchangeTextureFactoryNode* const* TextureFactoryNode = FactoryVariant.TryGet<UInterchangeTextureFactoryNode*>())
 		{
-			if ((*TextureFactoryNode)->GetCustomTranslatedTextureNodeUid(TextureNodeUniqueID))
-			{
-				if (const UInterchangeTexture2DNode* TextureTranslatedNode = Cast<UInterchangeTexture2DNode>(NodeContainer->GetNode(TextureNodeUniqueID)))
-				{
-					return FTextureNodeVariant(TInPlaceType<const UInterchangeTexture2DNode*>(), TextureTranslatedNode);
-				}
-			}
+			(*TextureFactoryNode)->GetCustomTranslatedTextureNodeUid(TextureNodeUniqueID);
 		}
 		else if (UInterchangeTextureCubeFactoryNode* const* TextureCubeFactoryNode = FactoryVariant.TryGet<UInterchangeTextureCubeFactoryNode*>())
 		{
-			if ((*TextureCubeFactoryNode)->GetCustomTranslatedTextureNodeUid(TextureNodeUniqueID))
-			{
-				if (const UInterchangeBaseNode* Node = NodeContainer->GetNode(TextureNodeUniqueID))
-				{
-					if (const UInterchangeTextureCubeNode* TextureCubeTranslatedNode = Cast<UInterchangeTextureCubeNode>(Node))
-					{
-						return FTextureNodeVariant(TInPlaceType<const UInterchangeTextureCubeNode*>(), TextureCubeTranslatedNode);
-					}
-					else if (const UInterchangeTexture2DNode* TextureTranslatedNode = Cast<UInterchangeTexture2DNode>(Node))
-					{
-						return FTextureNodeVariant(TInPlaceType<const UInterchangeTexture2DNode*>(), TextureTranslatedNode);
-					}
-				}
-			}
+			(*TextureCubeFactoryNode)->GetCustomTranslatedTextureNodeUid(TextureNodeUniqueID);
 		}
 		else if (UInterchangeTexture2DArrayFactoryNode* const* Texture2DArrayFactoryNode = FactoryVariant.TryGet<UInterchangeTexture2DArrayFactoryNode*>())
 		{
-			if ((*Texture2DArrayFactoryNode)->GetCustomTranslatedTextureNodeUid(TextureNodeUniqueID))
+			(*Texture2DArrayFactoryNode)->GetCustomTranslatedTextureNodeUid(TextureNodeUniqueID);
+		}
+		else if (UInterchangeTextureLightProfileFactoryNode* const* TextureLightProfileFactoryNode = FactoryVariant.TryGet<UInterchangeTextureLightProfileFactoryNode*>())
+		{
+			(*TextureLightProfileFactoryNode)->GetCustomTranslatedTextureNodeUid(TextureNodeUniqueID);
+		}
+
+		if (const UInterchangeBaseNode* TranslatedNode = NodeContainer->GetNode(TextureNodeUniqueID))
+		{
+			if (const UInterchangeTextureCubeNode* TextureCubeTranslatedNode = Cast<UInterchangeTextureCubeNode>(TranslatedNode))
 			{
-				if (const UInterchangeTexture2DArrayNode* Texture2DArrayTranslatedNode = Cast<UInterchangeTexture2DArrayNode>(NodeContainer->GetNode(TextureNodeUniqueID)))
-				{
-					return FTextureNodeVariant(TInPlaceType<const UInterchangeTexture2DArrayNode*>(), Texture2DArrayTranslatedNode);
-				}
+				return FTextureNodeVariant(TInPlaceType<const UInterchangeTextureCubeNode*>(), TextureCubeTranslatedNode);
+			}
+
+			if (const UInterchangeTexture2DArrayNode* Texture2DArrayTranslatedNode = Cast<UInterchangeTexture2DArrayNode>(TranslatedNode))
+			{
+				return FTextureNodeVariant(TInPlaceType<const UInterchangeTexture2DArrayNode*>(), Texture2DArrayTranslatedNode);
+			}
+
+			if (const UInterchangeTextureLightProfileNode* TextureLightProfileTranslatedNode = Cast<UInterchangeTextureLightProfileNode>(TranslatedNode))
+			{
+				return FTextureNodeVariant(TInPlaceType<const UInterchangeTextureLightProfileNode*>(), TextureLightProfileTranslatedNode);
+			}
+
+			if (const UInterchangeTexture2DNode* TextureTranslatedNode = Cast<UInterchangeTexture2DNode>(TranslatedNode))
+			{
+				return FTextureNodeVariant(TInPlaceType<const UInterchangeTexture2DNode*>(), TextureTranslatedNode);
 			}
 		}
 
@@ -174,7 +194,12 @@ namespace UE::Interchange::Private::InterchangeTextureFactory
 			return (*Texture2DArrayNode)->GetPayLoadKey().IsSet();
 		}
 
-	return false;
+		if (const UInterchangeTextureLightProfileNode* const* TextureLightProfileNode = TextureNodeVariant.TryGet<const UInterchangeTextureLightProfileNode*>())
+		{
+			return (*TextureLightProfileNode)->GetPayLoadKey().IsSet();
+		}
+
+		return false;
 	}
 
 	TOptional<FString> GetPayloadKey(const FTextureNodeVariant& TextureNodeVariant)
@@ -194,10 +219,19 @@ namespace UE::Interchange::Private::InterchangeTextureFactory
 			return (*Texture2DArrayNode)->GetPayLoadKey();
 		}
 
+		if (const UInterchangeTextureLightProfileNode* const* TextureLightProfileNode = TextureNodeVariant.TryGet<const UInterchangeTextureLightProfileNode*>())
+		{
+			return (*TextureLightProfileNode)->GetPayLoadKey();
+		}
+
 		return {};
 	}
 
-	using FTexturePayloadVariant = TVariant<FEmptyVariantState, TOptional<FImportImage>,  TOptional<FImportBlockedImage>,  TOptional<FImportSlicedImage>>;
+	using FTexturePayloadVariant = TVariant<FEmptyVariantState
+		, TOptional<FImportImage>
+		, TOptional<FImportBlockedImage>
+		, TOptional<FImportSlicedImage>
+		, TOptional<FImportLightProfile>>;
 
 	FTexturePayloadVariant GetTexturePayload(const UInterchangeSourceData* SourceData, const FString& PayloadKey, const FTextureNodeVariant& TextureNodeVariant, const UInterchangeTranslatorBase* Translator)
 	{
@@ -223,6 +257,15 @@ namespace UE::Interchange::Private::InterchangeTextureFactory
 			}
 		}
 
+		// Light Profile
+		if (TextureNodeVariant.IsType<const UInterchangeTextureLightProfileNode*>())
+		{
+			if (const IInterchangeTextureLightProfilePayloadInterface* LightProfileTranslator = Cast<IInterchangeTextureLightProfilePayloadInterface>(Translator))
+			{
+				return FTexturePayloadVariant(TInPlaceType<TOptional<FImportLightProfile>>(), LightProfileTranslator->GetLightProfilePayloadData(SourceData, PayloadKey));
+			}
+		}
+
 		return {};
 	}
 
@@ -300,27 +343,6 @@ namespace UE::Interchange::Private::InterchangeTextureFactory
 		return false;
 	}
 
-	bool SetupTexture2DSourceData(UTexture2D* Texture2D, const FTexturePayloadVariant& TexturePayload)
-	{
-		if (const TOptional<FImportBlockedImage>* BlockedImage = TexturePayload.TryGet<TOptional<FImportBlockedImage>>())
-		{
-			if (BlockedImage->IsSet())
-			{
-				return SetupTexture2DSourceData(Texture2D, BlockedImage->GetValue());
-			}
-		}
-		else if (const TOptional<FImportImage>* Image = TexturePayload.TryGet<TOptional<FImportImage>>())
-		{
-			if (Image->IsSet())
-			{
-				SetupTextureSourceData(Texture2D, Image->GetValue());
-				return true;
-			}
-		}
-
-		return false;
-	}
-
 	void SetupTextureSourceData(UTexture* Texture, const FImportSlicedImage& SlicedImage)
 	{
 		Texture->Source.Init(
@@ -365,6 +387,51 @@ namespace UE::Interchange::Private::InterchangeTextureFactory
 		}
 	}
 
+	void SetupTextureSourceData(UTextureLightProfile* TextureLightProfile, const FImportLightProfile& LightProfile)
+	{
+		const FImportImage& ImportImage = LightProfile;
+		SetupTextureSourceData(TextureLightProfile, ImportImage);
+
+		TextureLightProfile->Brightness = LightProfile.Brightness;
+		TextureLightProfile->TextureMultiplier = LightProfile.TextureMultiplier;
+	}
+
+	bool SetupTexture2DSourceData(UTexture2D* Texture2D, const FTexturePayloadVariant& TexturePayload)
+	{
+		if (const TOptional<FImportBlockedImage>* BlockedImage = TexturePayload.TryGet<TOptional<FImportBlockedImage>>())
+		{
+			if (BlockedImage->IsSet())
+			{
+				return SetupTexture2DSourceData(Texture2D, BlockedImage->GetValue());
+			}
+		}
+		else if (const TOptional<FImportImage>* Image = TexturePayload.TryGet<TOptional<FImportImage>>())
+		{
+			if (Image->IsSet())
+			{
+				SetupTextureSourceData(Texture2D, Image->GetValue());
+				return true;
+			}
+		}
+		else if (const TOptional<FImportLightProfile>* OptionalLightProfile = TexturePayload.TryGet<TOptional<FImportLightProfile>>())
+		{
+			const FImportLightProfile& LightProfile = OptionalLightProfile->GetValue();
+
+			if (UTextureLightProfile* TextureLightProfile = Cast<UTextureLightProfile>(Texture2D))
+			{
+				SetupTextureSourceData(TextureLightProfile, LightProfile);
+			}
+			else
+			{
+				SetupTextureSourceData(Texture2D, LightProfile);
+			}
+
+			return true;
+		}
+
+		return false;
+	}
+
 	bool SetupTextureCubeSourceData(UTextureCube* TextureCube, const FTexturePayloadVariant& TexturePayload)
 	{
 		if (const TOptional<FImportSlicedImage>* OptionalSlicedImage = TexturePayload.TryGet<TOptional<FImportSlicedImage>>())
@@ -390,6 +457,15 @@ namespace UE::Interchange::Private::InterchangeTextureFactory
 				return true;
 			}
 		}
+		else if (const TOptional<FImportLightProfile>* OptionalLightProfile = TexturePayload.TryGet<TOptional<FImportLightProfile>>())
+		{
+			if (OptionalLightProfile->IsSet())
+			{
+				const FImportLightProfile& LightProfile = OptionalLightProfile->GetValue();
+				SetupTextureSourceData(TextureCube, LightProfile);
+				return true;
+			}
+		}
 
 		return false;
 	}
@@ -415,9 +491,16 @@ namespace UE::Interchange::Private::InterchangeTextureFactory
 				return true;
 			}
 		}
+		else if (const TOptional<FImportLightProfile>* OptionalLightProfile = TexturePayload.TryGet<TOptional<FImportLightProfile>>())
+		{
+			const FImportLightProfile& LightProfile = OptionalLightProfile->GetValue();
+			SetupTextureSourceData(Texture2DArray, LightProfile);
+			return true;
+		}
 
 		return false;
 	}
+
 #endif // WITH_EDITORONLY_DATA
  }
 
@@ -582,7 +665,7 @@ UObject* UInterchangeTextureFactory::CreateAsset(const UInterchangeTextureFactor
 	{
 		if (!SetupTexture2DSourceData(Texture2D, TexturePayload))
 		{
-			UE_LOG(LogInterchangeImport, Error, TEXT("UInterchangeTextureFactory: The Payload was invalid for a 2D Texture"), *Arguments.AssetName);
+			UE_LOG(LogInterchangeImport, Error, TEXT("UInterchangeTextureFactory: The Payload was invalid for a %s. (%s)"), *(TextureClass->GetName()), *Arguments.AssetName);
 			return nullptr;
 		}
 	}
@@ -590,7 +673,7 @@ UObject* UInterchangeTextureFactory::CreateAsset(const UInterchangeTextureFactor
 	{
 		if (!SetupTextureCubeSourceData(TextureCube, TexturePayload))
 		{
-			UE_LOG(LogInterchangeImport, Error, TEXT("UInterchangeTextureFactory: The Payload was invalid for a TextureCube"), *Arguments.AssetName);
+			UE_LOG(LogInterchangeImport, Error, TEXT("UInterchangeTextureFactory: The Payload was invalid for a TextureCube. (%s)"), *Arguments.AssetName);
 			return nullptr;
 		}
 	}
@@ -598,7 +681,7 @@ UObject* UInterchangeTextureFactory::CreateAsset(const UInterchangeTextureFactor
 	{
 		if (!SetupTexture2DArraySourceData(Texture2DArray, TexturePayload))
 		{
-			UE_LOG(LogInterchangeImport, Error, TEXT("UInterchangeTextureFactory: The Payload was invalid for a TextureCube"), *Arguments.AssetName);
+			UE_LOG(LogInterchangeImport, Error, TEXT("UInterchangeTextureFactory: The Payload was invalid for a TextureArray. (%s)"), *Arguments.AssetName);
 			return nullptr;
 		}
 	}
diff --git a/Engine/Plugins/Runtime/Interchange/Source/Import/Private/Texture/TextureTranslatorUtilities.cpp b/Engine/Plugins/Runtime/Interchange/Source/Import/Private/Texture/TextureTranslatorUtilities.cpp
--- a/Engine/Plugins/Runtime/Interchange/Source/Import/Private/Texture/TextureTranslatorUtilities.cpp
+++ b/Engine/Plugins/Runtime/Interchange/Source/Import/Private/Texture/TextureTranslatorUtilities.cpp
@@ -6,6 +6,7 @@
 #include "InterchangeTexture2DArrayNode.h"
 #include "InterchangeTexture2DNode.h"
 #include "InterchangeTextureCubeNode.h"
+#include "InterchangeTextureLightProfileNode.h"
 #include "InterchangeTextureNode.h"
 #include "Misc/Paths.h"
 #include "Nodes/InterchangeBaseNodeContainer.h"
@@ -61,3 +62,8 @@ bool UE::Interchange::FTextureTranslatorUtilities::GenericTexture2DArrayTranslat
 {
 	return Private::GenericTextureTranslate(SourceData, BaseNodeContainer, UInterchangeTexture2DArrayNode::StaticClass());
 }
+
+bool UE::Interchange::FTextureTranslatorUtilities::GenericTextureLightProfileTranslate(const UInterchangeSourceData* SourceData, UInterchangeBaseNodeContainer& BaseNodeContainer)
+{
+	return Private::GenericTextureTranslate(SourceData, BaseNodeContainer, UInterchangeTextureLightProfileNode::StaticClass());
+}
diff --git a/Engine/Plugins/Runtime/Interchange/Source/Import/Private/Texture/TextureTranslatorUtilities.h b/Engine/Plugins/Runtime/Interchange/Source/Import/Private/Texture/TextureTranslatorUtilities.h
--- a/Engine/Plugins/Runtime/Interchange/Source/Import/Private/Texture/TextureTranslatorUtilities.h
+++ b/Engine/Plugins/Runtime/Interchange/Source/Import/Private/Texture/TextureTranslatorUtilities.h
@@ -24,6 +24,8 @@ namespace UE
 			static bool GenericTextureCubeTranslate(const UInterchangeSourceData* SourceData, UInterchangeBaseNodeContainer& BaseNodeContainer);
 
 			static bool GenericTexture2DArrayTranslate(const UInterchangeSourceData* SourceData, UInterchangeBaseNodeContainer& BaseNodeContainer);
+
+			static bool GenericTextureLightProfileTranslate(const UInterchangeSourceData* SourceData, UInterchangeBaseNodeContainer& BaseNodeContainer);
 		};
 	}//ns Interchange
 }//ns UE
diff --git a/Engine/Plugins/Runtime/Interchange/Source/Import/Public/Texture/InterchangeIESTranslator.h b/Engine/Plugins/Runtime/Interchange/Source/Import/Public/Texture/InterchangeIESTranslator.h
new file mode 100644
--- /dev/null
+++ b/Engine/Plugins/Runtime/Interchange/Source/Import/Public/Texture/InterchangeIESTranslator.h
@@ -0,0 +1,49 @@
+// Copyright Epic Games, Inc. All Rights Reserved.
+
+#pragma once
+
+#include "CoreMinimal.h"
+#include "UObject/Object.h"
+#include "UObject/ObjectMacros.h"
+#include "InterchangeTranslatorBase.h"
+#include "Nodes/InterchangeBaseNodeContainer.h"
+#include "Texture/InterchangeTextureLightProfilePayloadData.h"
+#include "Texture/InterchangeTextureLightProfilePayloadInterface.h"
+
+#include "InterchangeIESTranslator.generated.h"
+
+UCLASS(BlueprintType)
+class INTERCHANGEIMPORT_API UInterchangeIESTranslator : public UInterchangeTranslatorBase, public IInterchangeTextureLightProfilePayloadInterface
+{
+	GENERATED_BODY()
+public:
+
+	/*
+	 * return true if the translator can translate the specified file.
+	 */
+	virtual bool CanImportSourceData(const UInterchangeSourceData* InSourceData) const override;
+
+	/**
+	 * Translate a source data into a node hold by the specified nodes container.
+	 *
+	 * @param SourceData - The source data containing the data to translate
+	 * @param BaseNodeContainer - The unreal objects descriptions container where to put the translated source data.
+	 * @return true if the translator can translate the source data, false otherwise.
+	 */
+	virtual bool Translate(const UInterchangeSourceData* SourceData, UInterchangeBaseNodeContainer& BaseNodeContainer) const override;
+
+
+	/* IInterchangeTexturePayloadInterface Begin */
+
+	/**
+	 * Once the translation is done, the import process need a way to retrieve payload data.
+	 * This payload will be use by the factories to create the asset.
+	 *
+	 * @param SourceData - The source data containing the data to translate
+	 * @param PayloadKey - The key to retrieve the a particular payload contain into the specified source data.
+	 * @return a PayloadData containing the import light profile data. The TOptional will not be set if there is an error.
+	 */
+	virtual TOptional<UE::Interchange::FImportLightProfile> GetLightProfilePayloadData(const UInterchangeSourceData* SourceData, const FString& PayLoadKey) const override;
+
+	/* IInterchangeTexturePayloadInterface End */
+};
\ No newline at end of file
diff --git a/Engine/Plugins/Runtime/Interchange/Source/Import/Public/Texture/InterchangeTextureLightProfilePayloadData.h b/Engine/Plugins/Runtime/Interchange/Source/Import/Public/Texture/InterchangeTextureLightProfilePayloadData.h
new file mode 100644
--- /dev/null
+++ b/Engine/Plugins/Runtime/Interchange/Source/Import/Public/Texture/InterchangeTextureLightProfilePayloadData.h
@@ -0,0 +1,15 @@
+// Copyright Epic Games, Inc. All Rights Reserved.
+
+#pragma once
+
+#include "Texture/InterchangeTexturePayloadData.h"
+
+namespace UE::Interchange
+{
+	struct INTERCHANGEIMPORT_API FImportLightProfile : public FImportImage
+	{
+		float Brightness = 0.0f;
+
+		float TextureMultiplier = 0.0f;
+	};
+}
\ No newline at end of file
diff --git a/Engine/Plugins/Runtime/Interchange/Source/Import/Public/Texture/InterchangeTextureLightProfilePayloadInterface.h b/Engine/Plugins/Runtime/Interchange/Source/Import/Public/Texture/InterchangeTextureLightProfilePayloadInterface.h
new file mode 100644
--- /dev/null
+++ b/Engine/Plugins/Runtime/Interchange/Source/Import/Public/Texture/InterchangeTextureLightProfilePayloadInterface.h
@@ -0,0 +1,35 @@
+// Copyright Epic Games, Inc. All Rights Reserved.
+#pragma once
+
+#include "CoreMinimal.h"
+#include "InterchangeSourceData.h"
+#include "Texture/InterchangeTextureLightProfilePayloadData.h"
+#include "UObject/Object.h"
+#include "UObject/ObjectMacros.h"
+
+#include "InterchangeTextureLightProfilePayloadInterface.generated.h"
+
+UINTERFACE()
+class INTERCHANGEIMPORT_API UInterchangeTextureLightProfilePayloadInterface : public UInterface
+{
+	GENERATED_BODY()
+};
+
+/**
+ * Texture payload interface. Derive from it if your payload can import texture
+ */
+class INTERCHANGEIMPORT_API IInterchangeTextureLightProfilePayloadInterface
+{
+	GENERATED_BODY()
+public:
+
+	/**
+	 * Once the translation is done, the import process need a way to retrieve payload data.
+	 * This payload will be use by the factories to create the asset.
+	 *
+	 * @param SourceData - The source data containing the data to translate
+	 * @param PayloadKey - The key to retrieve the a particular payload contain into the specified source data.
+	 * @return a PayloadData containing the import image data. The TOptional will not be set if there is an error.
+	 */
+	virtual TOptional<UE::Interchange::FImportLightProfile> GetLightProfilePayloadData(const UInterchangeSourceData* SourceData, const FString& PayloadKey) const = 0;
+};
diff --git a/Engine/Plugins/Runtime/Interchange/Source/Nodes/Public/InterchangeTextureLightProfileFactoryNode.h b/Engine/Plugins/Runtime/Interchange/Source/Nodes/Public/InterchangeTextureLightProfileFactoryNode.h
new file mode 100644
--- /dev/null
+++ b/Engine/Plugins/Runtime/Interchange/Source/Nodes/Public/InterchangeTextureLightProfileFactoryNode.h
@@ -0,0 +1,77 @@
+// Copyright Epic Games, Inc. All Rights Reserved.
+
+#pragma once
+
+#include "CoreMinimal.h"
+#include "InterchangeTextureFactoryNode.h"
+#include "Nodes/InterchangeBaseNode.h"
+
+#if WITH_ENGINE
+#include "Engine/Texture.h"
+#include "Engine/TextureLightProfile.h"
+#endif
+
+
+#include "InterchangeTextureLightProfileFactoryNode.generated.h"
+
+UCLASS(BlueprintType)
+class INTERCHANGENODES_API UInterchangeTextureLightProfileFactoryNode : public UInterchangeTextureFactoryNode
+{
+	GENERATED_BODY()
+
+public:
+
+	UFUNCTION(BlueprintCallable, Category = "Interchange | Node | Texture")
+	bool GetCustomBrightness(float& AttributeValue) const
+	{
+		IMPLEMENT_NODE_ATTRIBUTE_GETTER(Brightness, float);
+	}
+
+	UFUNCTION(BlueprintCallable, Category = "Interchange | Node | Texture")
+	bool SetCustomBrightness(const float AttributeValue, bool bAddApplyDelegate = true)
+	{
+#if WITH_EDITORONLY_DATA
+		IMPLEMENT_NODE_ATTRIBUTE_SETTER(UInterchangeTextureLightProfileFactoryNode, Brightness, float, UTextureLightProfile)
+#else
+		IMPLEMENT_NODE_ATTRIBUTE_SETTER_NODELEGATE(Brightness, float)
+#endif
+	}
+
+
+	UFUNCTION(BlueprintCallable, Category = "Interchange | Node | Texture")
+	bool GetCustomTextureMultiplier(float& AttributeValue) const
+	{
+		IMPLEMENT_NODE_ATTRIBUTE_GETTER(TextureMultiplier, float);
+	}
+
+	UFUNCTION(BlueprintCallable, Category = "Interchange | Node | Texture")
+	bool SetCustomTextureMultiplier(const float AttributeValue, bool bAddApplyDelegate = true)
+	{
+#if WITH_EDITORONLY_DATA
+		IMPLEMENT_NODE_ATTRIBUTE_SETTER(UInterchangeTextureLightProfileFactoryNode, TextureMultiplier, float, UTextureLightProfile)
+#else
+		IMPLEMENT_NODE_ATTRIBUTE_SETTER_NODELEGATE(TextureMultiplier, float)
+#endif
+	}
+
+
+private:
+#if WITH_ENGINE
+	virtual void FillAssetClassFromClassName(const FString& ClassName) override
+	{
+		if (ClassName.Equals(UTextureLightProfile::StaticClass()->GetName()))
+		{
+			AssetClass = UTextureLightProfile::StaticClass();
+			bIsTextureNodeClassInitialized = true;
+		}
+	}
+#endif
+
+	// Addressing
+	IMPLEMENT_NODE_ATTRIBUTE_KEY(Brightness);
+	IMPLEMENT_NODE_ATTRIBUTE_KEY(TextureMultiplier);
+
+	IMPLEMENT_NODE_ATTRIBUTE_APPLY_UOBJECT(Brightness, float, UTextureLightProfile,);
+	IMPLEMENT_NODE_ATTRIBUTE_APPLY_UOBJECT(TextureMultiplier, float, UTextureLightProfile,);
+
+};
\ No newline at end of file
diff --git a/Engine/Plugins/Runtime/Interchange/Source/Nodes/Public/InterchangeTextureLightProfileNode.h b/Engine/Plugins/Runtime/Interchange/Source/Nodes/Public/InterchangeTextureLightProfileNode.h
new file mode 100644
--- /dev/null
+++ b/Engine/Plugins/Runtime/Interchange/Source/Nodes/Public/InterchangeTextureLightProfileNode.h
@@ -0,0 +1,26 @@
+// Copyright Epic Games, Inc. All Rights Reserved.
+
+#pragma once
+
+#include "CoreMinimal.h"
+#include "InterchangeTextureNode.h"
+
+#include "InterchangeTextureLightProfileNode.generated.h"
+
+
+UCLASS(BlueprintType)
+class INTERCHANGENODES_API UInterchangeTextureLightProfileNode : public UInterchangeTextureNode
+{
+	GENERATED_BODY()
+
+public:
+
+	/**
+	 * Return the node type name of the class, we use this when reporting error
+	 */
+	virtual FString GetTypeName() const override
+	{
+		const FString TypeName = TEXT("TextureLightProfileNode");
+		return TypeName;
+	}
+};
\ No newline at end of file
diff --git a/Engine/Plugins/Runtime/Interchange/Source/Pipelines/Private/InterchangeGenericAssetsPipeline.cpp b/Engine/Plugins/Runtime/Interchange/Source/Pipelines/Private/InterchangeGenericAssetsPipeline.cpp
--- a/Engine/Plugins/Runtime/Interchange/Source/Pipelines/Private/InterchangeGenericAssetsPipeline.cpp
+++ b/Engine/Plugins/Runtime/Interchange/Source/Pipelines/Private/InterchangeGenericAssetsPipeline.cpp
@@ -4,6 +4,7 @@
 #include "Animation/Skeleton.h"
 #include "CoreMinimal.h"
 #include "Engine/TextureCube.h"
+#include "Engine/TextureLightProfile.h"
 #include "InterchangeMaterialFactoryNode.h"
 #include "InterchangeMaterialNode.h"
 #include "InterchangeMeshNode.h"
@@ -20,6 +21,8 @@
 #include "InterchangeTextureCubeFactoryNode.h"
 #include "InterchangeTextureCubeNode.h"
 #include "InterchangeTextureFactoryNode.h"
+#include "InterchangeTextureLightProfileFactoryNode.h"
+#include "InterchangeTextureLightProfileNode.h"
 #include "InterchangeTextureNode.h"
 #include "Misc/Paths.h"
 #include "Nodes/InterchangeBaseNode.h"
@@ -47,6 +50,11 @@ namespace UE::Interchange::Private
 			return UInterchangeTexture2DArrayFactoryNode::StaticClass();
 		}
 
+		if (UInterchangeTextureLightProfileNode::StaticClass() == NodeClass)
+		{
+			return UInterchangeTextureLightProfileFactoryNode::StaticClass();
+		}
+
 		return nullptr;
 	}
 
@@ -67,6 +75,11 @@ namespace UE::Interchange::Private
 			return UTexture2DArray::StaticClass();
 		}
 
+		if (UInterchangeTextureLightProfileFactoryNode::StaticClass() == NodeClass)
+		{
+			return UTextureLightProfile::StaticClass();
+		}
+
 		return nullptr;
 	}
 }
