First pass implementaiton of a message passing system for Interchange. This permits messages, warnings and errors to be output from translators, pipelines and factories, as well as from external processes via the Interchange dispatcher.
Implemented UI tab widget for dislaying import results.
Implemented basic static mesh import, excluding collision import.
Various API tweaks, to accomodate the message passing, plus various bugfixes.

diff --git a/Engine/Plugins/Editor/InterchangeEditor/InterchangeEditor.uplugin b/Engine/Plugins/Editor/InterchangeEditor/InterchangeEditor.uplugin
--- a/Engine/Plugins/Editor/InterchangeEditor/InterchangeEditor.uplugin
+++ b/Engine/Plugins/Editor/InterchangeEditor/InterchangeEditor.uplugin
@@ -24,6 +24,11 @@
 			"Name": "InterchangeEditorPipelines",
 			"Type": "Editor",
 			"LoadingPhase": "Default"
+		},
+		{
+			"Name": "InterchangeResultsBrowser",
+			"Type": "Editor",
+			"LoadingPhase": "Default"
 		}
 	],
 	"Plugins": [
diff --git a/Engine/Plugins/Editor/InterchangeEditor/Source/InterchangeEditor/InterchangeEditor.Build.cs b/Engine/Plugins/Editor/InterchangeEditor/Source/InterchangeEditor/InterchangeEditor.Build.cs
--- a/Engine/Plugins/Editor/InterchangeEditor/Source/InterchangeEditor/InterchangeEditor.Build.cs
+++ b/Engine/Plugins/Editor/InterchangeEditor/Source/InterchangeEditor/InterchangeEditor.Build.cs
@@ -19,12 +19,7 @@ namespace UnrealBuildTool.Rules
 			PrivateDependencyModuleNames.AddRange(
 				new string[]
 				{
-					"InterchangeCore",
-					"InterchangeImport",
-					"InterchangeNodes",
-					"InterchangeExport",
-					"InterchangeDispatcher",
-					"InterchangeFbxParser"
+					"InterchangeCore"
 				}
 			);
 		}
diff --git a/Engine/Plugins/Editor/InterchangeEditor/Source/Pipelines/InterchangeEditorPipelines.Build.cs b/Engine/Plugins/Editor/InterchangeEditor/Source/Pipelines/InterchangeEditorPipelines.Build.cs
--- a/Engine/Plugins/Editor/InterchangeEditor/Source/Pipelines/InterchangeEditorPipelines.Build.cs
+++ b/Engine/Plugins/Editor/InterchangeEditor/Source/Pipelines/InterchangeEditorPipelines.Build.cs
@@ -13,8 +13,7 @@ namespace UnrealBuildTool.Rules
 					"CoreUObject",
 					"Engine",
 					"InterchangeCore",
-					"InterchangeEngine",
-					"InterchangeNodes",
+					"InterchangeEngine"
 				}
 			);
 
diff --git a/Engine/Plugins/Editor/InterchangeEditor/Source/Pipelines/Private/InterchangeGraphInspectorPipeline.cpp b/Engine/Plugins/Editor/InterchangeEditor/Source/Pipelines/Private/InterchangeGraphInspectorPipeline.cpp
--- a/Engine/Plugins/Editor/InterchangeEditor/Source/Pipelines/Private/InterchangeGraphInspectorPipeline.cpp
+++ b/Engine/Plugins/Editor/InterchangeEditor/Source/Pipelines/Private/InterchangeGraphInspectorPipeline.cpp
@@ -2,9 +2,6 @@
 #include "InterchangeGraphInspectorPipeline.h"
 
 #include "CoreMinimal.h"
-#include "InterchangeSkeletalMeshLodDataNode.h"
-#include "InterchangeSkeletalMeshFactoryNode.h"
-#include "InterchangeSkeletonFactoryNode.h"
 #include "InterchangeSourceData.h"
 #include "Interfaces/IMainFrameModule.h"
 #include "Misc/Paths.h"
diff --git a/Engine/Plugins/Editor/InterchangeEditor/Source/ResultsBrowser/InterchangeResultsBrowser.Build.cs b/Engine/Plugins/Editor/InterchangeEditor/Source/ResultsBrowser/InterchangeResultsBrowser.Build.cs
new file mode 100644
--- /dev/null
+++ b/Engine/Plugins/Editor/InterchangeEditor/Source/ResultsBrowser/InterchangeResultsBrowser.Build.cs
@@ -0,0 +1,33 @@
+// Copyright Epic Games, Inc. All Rights Reserved.
+
+namespace UnrealBuildTool.Rules
+{
+	public class InterchangeResultsBrowser : ModuleRules
+	{
+		public InterchangeResultsBrowser(ReadOnlyTargetRules Target) : base(Target)
+		{
+			PublicDependencyModuleNames.AddRange(
+				new string[]
+				{
+					"Core",
+					"CoreUObject",
+					"Engine",
+					"InterchangeCore",
+					"InterchangeEngine",
+				}
+			);
+
+			PrivateDependencyModuleNames.AddRange(
+				new string[]
+				{
+					"EditorStyle",
+					"InputCore",
+					"Projects",
+					"Slate",
+					"SlateCore",
+					"WorkspaceMenuStructure"
+				}
+			);
+		}
+	}
+}
diff --git a/Engine/Plugins/Editor/InterchangeEditor/Source/ResultsBrowser/Private/InterchangeResultsBrowserModule.cpp b/Engine/Plugins/Editor/InterchangeEditor/Source/ResultsBrowser/Private/InterchangeResultsBrowserModule.cpp
new file mode 100644
--- /dev/null
+++ b/Engine/Plugins/Editor/InterchangeEditor/Source/ResultsBrowser/Private/InterchangeResultsBrowserModule.cpp
@@ -0,0 +1,117 @@
+// Copyright Epic Games, Inc. All Rights Reserved.
+
+#include "InterchangeResultsBrowserModule.h"
+
+#include "CoreMinimal.h"
+#include "Engine/Engine.h"
+#include "Framework/Docking/TabManager.h"
+#include "InterchangeResultsBrowserStyle.h"
+#include "InterchangeManager.h"
+#include "SInterchangeResultsBrowserWindow.h"
+#include "Styling/SlateIconFinder.h"
+#include "Widgets/Docking/SDockTab.h"
+#include "WorkspaceMenuStructure.h"
+#include "WorkspaceMenuStructureModule.h"
+
+
+#define LOCTEXT_NAMESPACE "InterchangeResultsBrowser"
+
+class FInterchangeResultsBrowserModule : public IInterchangeResultsBrowserModule
+{
+	/** IModuleInterface implementation */
+	virtual void StartupModule() override;
+	virtual void ShutdownModule() override;
+
+private:
+
+	TSharedRef<SDockTab> SpawnInterchangeResults(const FSpawnTabArgs& Args);
+	void OpenErrorBrowser(TStrongObjectPtr<UInterchangeResultsContainer> InResultsContainer);
+
+	/** Pointer to the style set to use for the UI. */
+	TSharedPtr<ISlateStyle> InterchangeResultsBrowserStyle = nullptr;
+
+	/** Owning pointer to the InterchangeResultsContainer being displayed */
+	TStrongObjectPtr<UInterchangeResultsContainer> ResultsContainer;
+
+	/** Pointer to the Slate widget */
+	TSharedPtr<SInterchangeResultsBrowserWindow> InterchangeResultsBrowserWindow;
+
+	/** Whether the results list is filtered */
+	// @TODO: load this from config
+	bool bIsFiltered = false;
+};
+
+
+IMPLEMENT_MODULE(FInterchangeResultsBrowserModule, InterchangeResultsBrowser)
+
+
+void FInterchangeResultsBrowserModule::StartupModule()
+{
+	auto RegisterItems = [this]()
+	{
+		UInterchangeManager& InterchangeManager = UInterchangeManager::GetInterchangeManager();
+		InterchangeManager.OnBatchImportComplete.AddRaw(this, &FInterchangeResultsBrowserModule::OpenErrorBrowser);
+
+		auto UnregisterItems = [this]()
+		{
+			UInterchangeManager& InterchangeManager = UInterchangeManager::GetInterchangeManager();
+			InterchangeManager.OnBatchImportComplete.RemoveAll(this);
+		};
+
+		InterchangeManager.OnPreDestroyInterchangeManager.AddLambda(UnregisterItems);
+	};
+
+	if (GEngine)
+	{
+		RegisterItems();
+	}
+	else
+	{
+		FCoreDelegates::OnPostEngineInit.AddLambda(RegisterItems);
+	}
+	
+	if (!InterchangeResultsBrowserStyle.IsValid())
+	{
+		InterchangeResultsBrowserStyle = MakeShared<FInterchangeResultsBrowserStyle>();
+	}
+
+	FGlobalTabmanager::Get()->RegisterNomadTabSpawner("InterchangeResults", FOnSpawnTab::CreateRaw(this, &FInterchangeResultsBrowserModule::SpawnInterchangeResults))
+		.SetDisplayName(LOCTEXT("InterchangeResultsBrowser", "Interchange Results Browser"))
+		.SetTooltipText(LOCTEXT("InterchangeResultsBrowserTooltipText", "Open the Interchange Results Browser tab."))
+		.SetGroup(WorkspaceMenu::GetMenuStructure().GetDeveloperToolsLogCategory())
+		.SetIcon(FSlateIconFinder::FindIcon("InterchangeResultsBrowser.TabIcon"));
+}
+
+
+void FInterchangeResultsBrowserModule::ShutdownModule()
+{
+	InterchangeResultsBrowserStyle = nullptr;
+}
+
+
+TSharedRef<SDockTab> FInterchangeResultsBrowserModule::SpawnInterchangeResults(const FSpawnTabArgs& Args)
+{
+	TSharedRef<SDockTab> DockTab = SNew(SDockTab)
+		.TabRole(ETabRole::NomadTab);
+
+	DockTab->SetContent(
+		SAssignNew(InterchangeResultsBrowserWindow, SInterchangeResultsBrowserWindow)
+		.OwnerTab(DockTab)
+		.InterchangeResultsContainer(ResultsContainer.Get())
+		.IsFiltered(bIsFiltered)
+		.OnFilterChangedState_Lambda([this](bool bNewState) { this->bIsFiltered = bNewState; })
+	);
+
+	return DockTab;
+}
+
+
+void FInterchangeResultsBrowserModule::OpenErrorBrowser(TStrongObjectPtr<UInterchangeResultsContainer> InResultsContainer)
+{
+	ResultsContainer = InResultsContainer;
+	ResultsContainer->Finalize();
+	FGlobalTabmanager::Get()->TryInvokeTab(FName("InterchangeResults"));
+	InterchangeResultsBrowserWindow->Set(ResultsContainer.Get());
+}
+
+#undef LOCTEXT_NAMESPACE
diff --git a/Engine/Plugins/Editor/InterchangeEditor/Source/ResultsBrowser/Private/InterchangeResultsBrowserStyle.h b/Engine/Plugins/Editor/InterchangeEditor/Source/ResultsBrowser/Private/InterchangeResultsBrowserStyle.h
new file mode 100644
--- /dev/null
+++ b/Engine/Plugins/Editor/InterchangeEditor/Source/ResultsBrowser/Private/InterchangeResultsBrowserStyle.h
@@ -0,0 +1,44 @@
+// Copyright Epic Games, Inc. All Rights Reserved.
+
+#pragma once
+
+#include "Interfaces/IPluginManager.h"
+#include "Brushes/SlateImageBrush.h"
+#include "Styling/SlateStyle.h"
+#include "Styling/SlateStyleRegistry.h"
+
+/**
+ * Implements the visual style of the text asset editor UI.
+ */
+class FInterchangeResultsBrowserStyle : public FSlateStyleSet
+{
+public:
+
+	FInterchangeResultsBrowserStyle()
+		: FSlateStyleSet("InterchangeResultsBrowserStyle")
+	{
+		const FVector2D Icon16x16(16.0f, 16.0f);
+
+		const FString BaseDir = IPluginManager::Get().FindPlugin("InterchangeEditor")->GetBaseDir();
+		SetContentRoot(BaseDir / TEXT("Content"));
+
+		FSlateImageBrush* SuccessBrush = new FSlateImageBrush(RootToContentDir(TEXT("Resources/Result_Success"), TEXT(".png")), Icon16x16);
+		Set("InterchangeResultsBrowser.ResultType.Success", SuccessBrush);
+		
+		FSlateImageBrush* WarningBrush = new FSlateImageBrush(RootToContentDir(TEXT("Resources/Result_Warning"), TEXT(".png")), Icon16x16);
+		Set("InterchangeResultsBrowser.ResultType.Warning", WarningBrush);
+
+		FSlateImageBrush* ErrorBrush = new FSlateImageBrush(RootToContentDir(TEXT("Resources/Result_Fail"), TEXT(".png")), Icon16x16);
+		Set("InterchangeResultsBrowser.ResultType.Error", ErrorBrush);
+
+		FSlateImageBrush* TabIcon = new FSlateImageBrush(RootToContentDir(TEXT("Resources/TabIcon"), TEXT(".png")), Icon16x16);
+		Set("InterchangeResultsBrowser.TabIcon", TabIcon);
+
+		FSlateStyleRegistry::RegisterSlateStyle(*this);
+	}
+
+	~FInterchangeResultsBrowserStyle()
+	{
+		FSlateStyleRegistry::UnRegisterSlateStyle(*this);
+	}
+};
diff --git a/Engine/Plugins/Editor/InterchangeEditor/Source/ResultsBrowser/Private/SInterchangeResultsBrowserWindow.cpp b/Engine/Plugins/Editor/InterchangeEditor/Source/ResultsBrowser/Private/SInterchangeResultsBrowserWindow.cpp
new file mode 100644
--- /dev/null
+++ b/Engine/Plugins/Editor/InterchangeEditor/Source/ResultsBrowser/Private/SInterchangeResultsBrowserWindow.cpp
@@ -0,0 +1,434 @@
+// Copyright Epic Games, Inc. All Rights Reserved.
+
+#include "SInterchangeResultsBrowserWindow.h"
+
+#include "EditorStyleSet.h"
+#include "Framework/MultiBox/MultiBoxBuilder.h"
+#include "IDocumentation.h"
+#include "Modules/ModuleManager.h"
+#include "Styling/SlateIconFinder.h"
+#include "Widgets/Docking/SDockTab.h"
+#include "Widgets/Images/SImage.h"
+#include "Widgets/Input/SButton.h"
+#include "Widgets/Input/SCheckBox.h"
+#include "Widgets/Layout/SSplitter.h"
+#include "Widgets/Layout/SUniformGridPanel.h"
+#include "Widgets/Layout/SSpacer.h"
+
+#define LOCTEXT_NAMESPACE "InterchangeResultsBrowser"
+
+
+namespace SResultsBrowserWidgetDefs
+{
+	const FName ColumnID_ResultType("ResultType");
+	const FName ColumnID_Source("Source");
+	const FName ColumnID_Dest("Dest");
+	const FName ColumnID_AssetType("AssetType");
+	const FName ColumnID_Message("Message");
+}
+
+
+SInterchangeResultsBrowserWindow::SInterchangeResultsBrowserWindow()
+	: OwnerTab(nullptr),
+	  bIsFiltered(false)
+{
+}
+
+
+SInterchangeResultsBrowserWindow::~SInterchangeResultsBrowserWindow()
+{
+}
+
+
+void SInterchangeResultsBrowserWindow::Construct(const FArguments& InArgs)
+{
+	SortByColumn = SResultsBrowserWidgetDefs::ColumnID_ResultType;
+	SortMode = EColumnSortMode::Ascending;
+	OnFilterChangedState = InArgs._OnFilterChangedState;
+	bIsFiltered = InArgs._IsFiltered;
+	ResultsContainer = InArgs._InterchangeResultsContainer;
+
+	OwnerTab = InArgs._OwnerTab;
+	check(OwnerTab.IsValid());
+
+	TSharedRef<SHeaderRow> HeaderRowWidget = SNew(SHeaderRow);
+
+	HeaderRowWidget->AddColumn(
+		SHeaderRow::Column(SResultsBrowserWidgetDefs::ColumnID_ResultType)
+		.SortMode(this, &SInterchangeResultsBrowserWindow::GetColumnSortMode, SResultsBrowserWidgetDefs::ColumnID_ResultType)
+		.OnSort(this, &SInterchangeResultsBrowserWindow::OnColumnSortModeChanged)
+		[
+			SNew(SSpacer)
+		]
+		.FixedWidth(24.0f)
+		);
+
+	HeaderRowWidget->AddColumn(
+		SHeaderRow::Column(SResultsBrowserWidgetDefs::ColumnID_Source)
+		.DefaultLabel(LOCTEXT("ColumnLabel_SourceFile", "Source File"))
+		.SortMode(this, &SInterchangeResultsBrowserWindow::GetColumnSortMode, SResultsBrowserWidgetDefs::ColumnID_Source)
+		.OnSort(this, &SInterchangeResultsBrowserWindow::OnColumnSortModeChanged)
+		.FillWidth(5.0f)
+	);
+
+	HeaderRowWidget->AddColumn(
+		SHeaderRow::Column(SResultsBrowserWidgetDefs::ColumnID_Dest)
+		.DefaultLabel(LOCTEXT("ColumnLabel_DestinationAsset", "Destination Asset"))
+		.SortMode(this, &SInterchangeResultsBrowserWindow::GetColumnSortMode, SResultsBrowserWidgetDefs::ColumnID_Dest)
+		.OnSort(this, &SInterchangeResultsBrowserWindow::OnColumnSortModeChanged)
+		.FillWidth(5.0f)
+	);
+
+	HeaderRowWidget->AddColumn(
+		SHeaderRow::Column(SResultsBrowserWidgetDefs::ColumnID_AssetType)
+		.DefaultLabel(LOCTEXT("ColumnLabel_AssetType", "Asset Type"))
+		.SortMode(this, &SInterchangeResultsBrowserWindow::GetColumnSortMode, SResultsBrowserWidgetDefs::ColumnID_AssetType)
+		.OnSort(this, &SInterchangeResultsBrowserWindow::OnColumnSortModeChanged)
+		.FillWidth(3.0f)
+	);
+
+	HeaderRowWidget->AddColumn(
+		SHeaderRow::Column(SResultsBrowserWidgetDefs::ColumnID_Message)
+		.DefaultLabel(LOCTEXT("ColumnLabel_Details", "Details"))
+		.FillWidth(9.0f)
+	);
+
+	this->ChildSlot
+	[
+		SNew(SBorder)
+		.Padding(FMargin(10.0f, 3.0f))
+		.BorderImage(FEditorStyle::GetBrush("ToolPanel.DarkGroupBorder"))
+		[
+			SNew(SVerticalBox)
+			+ SVerticalBox::Slot()
+			.AutoHeight()
+			.HAlign(HAlign_Left)
+			.Padding(10)
+			[
+				SNew(SHorizontalBox)
+				+ SHorizontalBox::Slot()
+				.AutoWidth()
+				[
+					SNew(SCheckBox)
+					.IsChecked(bIsFiltered ? ECheckBoxState::Checked : ECheckBoxState::Unchecked)
+					.OnCheckStateChanged(this, &SInterchangeResultsBrowserWindow::OnFilterStateChanged)
+				]
+				+ SHorizontalBox::Slot()
+				.AutoWidth()
+				[
+					SNew(STextBlock)
+					.Text(LOCTEXT("OnlyShowErrors", "Only display errors/warnings"))
+				]
+			]
+			+ SVerticalBox::Slot()
+			.FillHeight(1.0f)
+			.Padding(2)
+			[
+				SNew(SSplitter)
+				.Orientation(Orient_Vertical)
+				+ SSplitter::Slot()
+				.Value(0.4f)
+				[
+					SAssignNew(ListView, SListView<UInterchangeResult*>)
+					.ItemHeight(30)
+					.ListItemsSource(&FilteredResults)
+					.HeaderRow(HeaderRowWidget)
+					.SelectionMode(ESelectionMode::Single)
+					.OnGenerateRow(this, &SInterchangeResultsBrowserWindow::OnGenerateRowForList)
+					.OnSelectionChanged(this, &SInterchangeResultsBrowserWindow::OnSelectionChanged)
+				]
+				+ SSplitter::Slot()
+				.Value(0.2f)
+				[
+					SNew(SBorder)
+					.Padding(3)
+					.BorderImage(FEditorStyle::GetBrush("ToolPanel.GroupBorder"))
+					[
+						SNew(SBox)
+						.Padding(10)
+						[
+							SAssignNew(Description, STextBlock)
+							.AutoWrapText(true)
+						]
+					]
+				]
+			]
+			+ SVerticalBox::Slot()
+			.AutoHeight()
+			.HAlign(HAlign_Right)
+			.Padding(2)
+			[
+				SNew(SUniformGridPanel)
+				.SlotPadding(2)
+				+ SUniformGridPanel::Slot(0, 0)
+				[
+					IDocumentation::Get()->CreateAnchor(FString("Engine/Content/Interchange/ErrorBrowser"))
+				]
+				+ SUniformGridPanel::Slot(1, 0)
+				[
+					SNew(SButton)
+					.HAlign(HAlign_Center)
+					.Text(LOCTEXT("ErrorBrowser_Clear", "Clear"))
+					.IsEnabled(false)
+					.OnClicked(this, &SInterchangeResultsBrowserWindow::OnCloseDialog)
+				]
+				+ SUniformGridPanel::Slot(2, 0)
+				[
+					SNew(SButton)
+					.HAlign(HAlign_Center)
+					.Text(LOCTEXT("ErrorBrowser_Close", "Close"))
+					.OnClicked(this, &SInterchangeResultsBrowserWindow::OnCloseDialog)
+				]
+			]
+		]
+	];
+
+	Set(InArgs._InterchangeResultsContainer);
+}
+
+
+void SInterchangeResultsBrowserWindow::Set(UInterchangeResultsContainer* Data)
+{
+	ResultsContainer = Data;
+	RepopulateItems();
+}
+
+
+void SInterchangeResultsBrowserWindow::RepopulateItems()
+{
+	FilteredResults.Empty();
+
+	if (ResultsContainer != nullptr)
+	{
+		for (UInterchangeResult* Item : ResultsContainer->GetResults())
+		{
+			if (Item->GetResultType() != EInterchangeResultType::Success || !bIsFiltered)
+			{
+				FilteredResults.Add(Item);
+			}
+		}
+
+		RequestSort();
+	}
+}
+
+
+void SInterchangeResultsBrowserWindow::CloseErrorBrowser()
+{
+	if (TSharedPtr<SDockTab> OwnerTabPin = OwnerTab.Pin())
+	{
+		OwnerTabPin->RequestCloseTab();
+	}
+	OwnerTab = nullptr;
+}
+
+
+TSharedRef<ITableRow> SInterchangeResultsBrowserWindow::OnGenerateRowForList(UInterchangeResult* InItem, const TSharedRef<STableViewBase>& OwnerTable)
+{
+	TSharedRef<ITableRow> Row =
+		SNew(SInterchangeResultsBrowserListRow, OwnerTable)
+		.InterchangeResultsBrowserWidget(SharedThis(this))
+		.Item(InItem);
+
+	return Row;
+}
+
+
+EColumnSortMode::Type SInterchangeResultsBrowserWindow::GetColumnSortMode(const FName ColumnId) const
+{
+	if (SortByColumn != ColumnId)
+	{
+		return EColumnSortMode::None;
+	}
+
+	return SortMode;
+}
+
+
+void SInterchangeResultsBrowserWindow::OnColumnSortModeChanged(const EColumnSortPriority::Type SortPriority, const FName& ColumnId, const EColumnSortMode::Type InSortMode)
+{
+	SortByColumn = ColumnId;
+	SortMode = InSortMode;
+
+	RequestSort();
+}
+
+
+void SInterchangeResultsBrowserWindow::OnFilterStateChanged(ECheckBoxState NewState)
+{
+	bIsFiltered = (NewState == ECheckBoxState::Checked);
+	OnFilterChangedState.ExecuteIfBound(bIsFiltered);
+	RepopulateItems();
+}
+
+
+void SInterchangeResultsBrowserWindow::RequestSort()
+{
+	if (SortByColumn == SResultsBrowserWidgetDefs::ColumnID_ResultType)
+	{
+		if (SortMode == EColumnSortMode::Ascending)
+		{
+			FilteredResults.Sort([](const UInterchangeResult& A, const UInterchangeResult& B) {
+				return static_cast<int>(A.GetResultType()) < static_cast<int>(B.GetResultType()); });
+		}
+		else if (SortMode == EColumnSortMode::Descending)
+		{
+			FilteredResults.Sort([](const UInterchangeResult& A, const UInterchangeResult& B) {
+				return static_cast<int>(A.GetResultType()) >= static_cast<int>(B.GetResultType()); });
+		}
+	}
+	else if (SortByColumn == SResultsBrowserWidgetDefs::ColumnID_Source)
+	{
+		if (SortMode == EColumnSortMode::Ascending)
+		{
+			FilteredResults.Sort([](const UInterchangeResult& A, const UInterchangeResult& B) {
+				return A.SourceAssetName < B.SourceAssetName; });
+		}
+		else if (SortMode == EColumnSortMode::Descending)
+		{
+			FilteredResults.Sort([](const UInterchangeResult& A, const UInterchangeResult& B) {
+				return A.SourceAssetName >= B.SourceAssetName; });
+		}
+	}
+	else if (SortByColumn == SResultsBrowserWidgetDefs::ColumnID_Dest)
+	{
+		if (SortMode == EColumnSortMode::Ascending)
+		{
+			FilteredResults.Sort([](const UInterchangeResult& A, const UInterchangeResult& B) {
+				return A.DestinationAssetName < B.DestinationAssetName; });
+		}
+		else if (SortMode == EColumnSortMode::Descending)
+		{
+			FilteredResults.Sort([](const UInterchangeResult& A, const UInterchangeResult& B) {
+				return A.DestinationAssetName >= B.DestinationAssetName; });
+		}
+	}
+	else if (SortByColumn == SResultsBrowserWidgetDefs::ColumnID_AssetType)
+	{
+		if (SortMode == EColumnSortMode::Ascending)
+		{
+			FilteredResults.Sort([](const UInterchangeResult& A, const UInterchangeResult& B) {
+				FString NameA = A.AssetType ? A.AssetType->GetName() : FString();
+				FString NameB = B.AssetType ? B.AssetType->GetName() : FString();
+				return NameA < NameB;
+				});
+		}
+		else if (SortMode == EColumnSortMode::Descending)
+		{
+			FilteredResults.Sort([](const UInterchangeResult& A, const UInterchangeResult& B) {
+				FString NameA = A.AssetType ? A.AssetType->GetName() : FString();
+				FString NameB = B.AssetType ? B.AssetType->GetName() : FString();
+				return NameA >= NameB;
+				});
+		}
+	}
+
+	ListView->RequestListRefresh();
+}
+
+
+
+void SInterchangeResultsBrowserWindow::OnSelectionChanged(UInterchangeResult* InItem, ESelectInfo::Type SelectInfo)
+{
+	if (InItem != nullptr)
+	{
+		Description->SetText(InItem->GetText());
+	}
+	else
+	{
+		Description->SetText(FText());
+	}
+}
+
+
+
+
+void SInterchangeResultsBrowserListRow::Construct(const FArguments& InArgs, const TSharedRef<STableViewBase>& InOwnerTableView)
+{
+	InterchangeResultsBrowserWidgetPtr = InArgs._InterchangeResultsBrowserWidget;
+	Item = InArgs._Item;
+
+	SMultiColumnTableRow<UInterchangeResult*>::Construct(FSuperRowType::FArguments(), InOwnerTableView);
+}
+
+
+static FName GetImageFromResultType(EInterchangeResultType ResultType)
+{
+	switch (ResultType)
+	{
+	case EInterchangeResultType::Success: return FName("InterchangeResultsBrowser.ResultType.Success");
+	case EInterchangeResultType::Warning: return FName("InterchangeResultsBrowser.ResultType.Warning");
+	case EInterchangeResultType::Error: return FName("InterchangeResultsBrowser.ResultType.Error");
+	}
+
+	return FName();
+}
+
+
+TSharedRef<SWidget> SInterchangeResultsBrowserListRow::GenerateWidgetForColumn(const FName& ColumnName)
+{
+	if (TSharedPtr<SInterchangeResultsBrowserWindow> InterchangeResultsBrowserWidget = InterchangeResultsBrowserWidgetPtr.Pin())
+	{
+		const FMargin RowPadding(8, 2, 2, 0);
+
+		if (ColumnName == SResultsBrowserWidgetDefs::ColumnID_ResultType)
+		{
+			return SNew(SHorizontalBox)
+				+ SHorizontalBox::Slot()
+				.Padding(RowPadding)
+				.AutoWidth()
+				[
+					SNew(SImage)
+					.Image(FSlateIconFinder::FindIcon(GetImageFromResultType(Item->GetResultType())).GetOptionalIcon())
+				];
+		}
+		else if (ColumnName == SResultsBrowserWidgetDefs::ColumnID_Source)
+		{
+			return SNew(SHorizontalBox)
+				+ SHorizontalBox::Slot()
+				.Padding(RowPadding)
+				[
+					SNew(STextBlock)
+					.Text(FText::FromString(FPaths::GetCleanFilename(Item->SourceAssetName)))
+					.ToolTipText(FText::FromString(Item->SourceAssetName))
+				];
+		}
+		else if (ColumnName == SResultsBrowserWidgetDefs::ColumnID_Dest)
+		{
+			return SNew(SHorizontalBox)
+				+ SHorizontalBox::Slot()
+				.Padding(RowPadding)
+				[
+					SNew(STextBlock)
+					.Text(FText::FromString(FPackageName::ObjectPathToObjectName(Item->DestinationAssetName)))
+					.ToolTipText(FText::FromString(Item->DestinationAssetName))
+				];
+		}
+		else if (ColumnName == SResultsBrowserWidgetDefs::ColumnID_AssetType)
+		{
+			return SNew(SHorizontalBox)
+				+ SHorizontalBox::Slot()
+				.Padding(RowPadding)
+				[
+					SNew(STextBlock)
+					.Text(FText::FromString(Item->AssetType ? Item->AssetType->GetName() : FString()))
+				];
+		}
+		else if (ColumnName == SResultsBrowserWidgetDefs::ColumnID_Message)
+		{
+			return SNew(SHorizontalBox)
+				+ SHorizontalBox::Slot()
+				.Padding(RowPadding)
+				[
+					SNew(STextBlock)
+					.Text(Item->GetText())
+					.ToolTipText(Item->GetText())
+				];
+		}
+	}
+
+	return SNullWidget::NullWidget;
+}
+
+
+#undef LOCTEXT_NAMESPACE
diff --git a/Engine/Plugins/Editor/InterchangeEditor/Source/ResultsBrowser/Private/SInterchangeResultsBrowserWindow.h b/Engine/Plugins/Editor/InterchangeEditor/Source/ResultsBrowser/Private/SInterchangeResultsBrowserWindow.h
new file mode 100644
--- /dev/null
+++ b/Engine/Plugins/Editor/InterchangeEditor/Source/ResultsBrowser/Private/SInterchangeResultsBrowserWindow.h
@@ -0,0 +1,102 @@
+// Copyright Epic Games, Inc. All Rights Reserved.
+
+#pragma once
+
+#include "CoreMinimal.h"
+#include "InputCoreTypes.h"
+#include "Input/Reply.h"
+#include "InterchangeResultsContainer.h"
+#include "Styling/SlateColor.h"
+#include "Widgets/DeclarativeSyntaxSupport.h"
+#include "Widgets/SCompoundWidget.h"
+#include "Widgets/Views/SListView.h"
+#include "Widgets/Views/STableViewBase.h"
+#include "Widgets/Views/STableRow.h"
+
+
+class SInterchangeResultsBrowserWindow : public SCompoundWidget
+{
+public:
+	DECLARE_DELEGATE_OneParam(FOnFilterChangedState, bool)
+
+	SLATE_BEGIN_ARGS(SInterchangeResultsBrowserWindow)
+		: _OwnerTab()
+		, _InterchangeResultsContainer()
+		, _IsFiltered(false)
+		{}
+
+		SLATE_ARGUMENT(TWeakPtr<SDockTab>, OwnerTab)
+		SLATE_ARGUMENT(UInterchangeResultsContainer*, InterchangeResultsContainer)
+		SLATE_ARGUMENT(bool, IsFiltered)
+		SLATE_EVENT(FOnFilterChangedState, OnFilterChangedState)
+	SLATE_END_ARGS()
+
+public:
+
+	SInterchangeResultsBrowserWindow();
+	~SInterchangeResultsBrowserWindow();
+
+	void Construct(const FArguments& InArgs);
+	virtual bool SupportsKeyboardFocus() const override { return true; }
+
+	void Set(UInterchangeResultsContainer* Data);
+	void CloseErrorBrowser();
+
+	FReply OnCloseDialog()
+	{
+		CloseErrorBrowser();
+		return FReply::Handled();
+	}
+
+	virtual FReply OnKeyDown(const FGeometry& MyGeometry, const FKeyEvent& InKeyEvent) override
+	{
+		if (InKeyEvent.GetKey() == EKeys::Escape)
+		{
+			return OnCloseDialog();
+		}
+
+		return FReply::Unhandled();
+	}
+
+private:
+
+	/** Called by SListView to get a widget corresponding to the supplied item */
+	TSharedRef<ITableRow> OnGenerateRowForList(UInterchangeResult* Item, const TSharedRef<STableViewBase>& OwnerTable);
+
+	EColumnSortMode::Type GetColumnSortMode(const FName ColumnId) const;
+
+	void OnColumnSortModeChanged(const EColumnSortPriority::Type SortPriority, const FName& ColumnId, const EColumnSortMode::Type InSortMode);
+	void OnSelectionChanged(UInterchangeResult* InItem, ESelectInfo::Type SelectInfo);
+	void OnFilterStateChanged(ECheckBoxState NewState);
+	void RepopulateItems();
+	void RequestSort();
+
+
+	TWeakPtr<SDockTab> OwnerTab;
+	UInterchangeResultsContainer* ResultsContainer;
+	TArray<UInterchangeResult*> FilteredResults;
+	TSharedPtr<SListView<UInterchangeResult*>> ListView;
+	TSharedPtr<STextBlock> Description;
+	FName SortByColumn;
+	EColumnSortMode::Type SortMode;
+	FOnFilterChangedState OnFilterChangedState;
+	bool bIsFiltered;
+};
+
+
+class SInterchangeResultsBrowserListRow : public SMultiColumnTableRow<UInterchangeResult*>
+{
+public:
+	SLATE_BEGIN_ARGS(SInterchangeResultsBrowserListRow) {}
+		SLATE_ARGUMENT(TSharedPtr<SInterchangeResultsBrowserWindow>, InterchangeResultsBrowserWidget)
+		SLATE_ARGUMENT(UInterchangeResult*, Item)
+	SLATE_END_ARGS()
+
+	void Construct(const FArguments& InArgs, const TSharedRef<STableViewBase>& InOwnerTableView);
+	virtual TSharedRef<SWidget> GenerateWidgetForColumn(const FName& ColumnName) override;
+
+private:
+
+	TWeakPtr<SInterchangeResultsBrowserWindow> InterchangeResultsBrowserWidgetPtr;
+	UInterchangeResult* Item;
+};
diff --git a/Engine/Plugins/Editor/InterchangeEditor/Source/ResultsBrowser/Public/InterchangeResultsBrowserModule.h b/Engine/Plugins/Editor/InterchangeEditor/Source/ResultsBrowser/Public/InterchangeResultsBrowserModule.h
new file mode 100644
--- /dev/null
+++ b/Engine/Plugins/Editor/InterchangeEditor/Source/ResultsBrowser/Public/InterchangeResultsBrowserModule.h
@@ -0,0 +1,39 @@
+// Copyright Epic Games, Inc. All Rights Reserved.
+
+#pragma once
+
+#include "CoreMinimal.h"
+#include "Modules/ModuleInterface.h"
+#include "Modules/ModuleManager.h"
+#include "Styling/ISlateStyle.h"
+
+
+/**
+ * The public interface to this module
+ */
+class IInterchangeResultsBrowserModule : public IModuleInterface
+{
+
+public:
+
+	/**
+	 * Singleton-like access to this module's interface.  This is just for convenience!
+	 * Beware of calling this during the shutdown phase, though.  Your module might have been unloaded already.
+	 *
+	 * @return Returns singleton instance, loading the module on demand if needed
+	 */
+	static inline IInterchangeResultsBrowserModule& Get()
+	{
+		return FModuleManager::LoadModuleChecked<IInterchangeResultsBrowserModule>("InterchangeResultsBrowser");
+	}
+
+	/**
+	 * Checks to see if this module is loaded and ready.  It is only valid to call Get() if IsAvailable() returns true.
+	 *
+	 * @return True if the module is loaded and ready to use
+	 */
+	static inline bool IsAvailable()
+	{
+		return FModuleManager::Get().IsModuleLoaded("InterchangeResultsBrowser");
+	}
+};
diff --git a/Engine/Plugins/Runtime/Interchange/Interchange.uplugin b/Engine/Plugins/Runtime/Interchange/Interchange.uplugin
--- a/Engine/Plugins/Runtime/Interchange/Interchange.uplugin
+++ b/Engine/Plugins/Runtime/Interchange/Interchange.uplugin
@@ -26,6 +26,11 @@
 			"Type": "Runtime",
 			"LoadingPhase": "Default"
 		},
+		{
+			"Name": "InterchangeMessages",
+			"Type": "Runtime",
+			"LoadingPhase": "Default"
+		},
 		{
 			"Name": "InterchangeExport",
 			"Type": "Runtime",
diff --git a/Engine/Plugins/Runtime/Interchange/Source/Import/InterchangeImport.Build.cs b/Engine/Plugins/Runtime/Interchange/Source/Import/InterchangeImport.Build.cs
--- a/Engine/Plugins/Runtime/Interchange/Source/Import/InterchangeImport.Build.cs
+++ b/Engine/Plugins/Runtime/Interchange/Source/Import/InterchangeImport.Build.cs
@@ -27,6 +27,7 @@ namespace UnrealBuildTool.Rules
 					"IESFile",
 					"ImageWrapper",
 					"InterchangeDispatcher",
+					"InterchangeMessages",
 					"Json",
 					"RHI",
 					"TextureUtilitiesCommon",
diff --git a/Engine/Plugins/Runtime/Interchange/Source/Import/Private/Fbx/InterchangeFbxTranslator.cpp b/Engine/Plugins/Runtime/Interchange/Source/Import/Private/Fbx/InterchangeFbxTranslator.cpp
--- a/Engine/Plugins/Runtime/Interchange/Source/Import/Private/Fbx/InterchangeFbxTranslator.cpp
+++ b/Engine/Plugins/Runtime/Interchange/Source/Import/Private/Fbx/InterchangeFbxTranslator.cpp
@@ -1,4 +1,5 @@
 // Copyright Epic Games, Inc. All Rights Reserved. 
+
 #include "Fbx/InterchangeFbxTranslator.h"
 
 #include "Dom/JsonValue.h"
@@ -21,6 +22,7 @@
 #include "Texture/InterchangeTexturePayloadData.h"
 #include "UObject/GCObjectScopeGuard.h"
 
+
 UInterchangeFbxTranslator::UInterchangeFbxTranslator(const class FObjectInitializer& ObjectInitializer)
 {
 	Dispatcher = nullptr;
@@ -34,9 +36,9 @@ bool UInterchangeFbxTranslator::CanImportSourceData(const UInterchangeSourceData
 	return FbxExtension.StartsWith(Extension,ESearchCase::IgnoreCase);
 }
 
-bool UInterchangeFbxTranslator::Translate(const UInterchangeSourceData* SourceData, UInterchangeBaseNodeContainer& BaseNodeContainer) const
+bool UInterchangeFbxTranslator::Translate(UInterchangeBaseNodeContainer& BaseNodeContainer) const
 {
-	FString Filename = SourceData->GetFilename();
+	FString Filename = GetSourceData()->GetFilename();
 	if (!FPaths::FileExists(Filename))
 	{
 		return false;
@@ -81,7 +83,7 @@ bool UInterchangeFbxTranslator::Translate(const UInterchangeSourceData* SourceDa
 	}
 
 	//Create a json command to read the fbx file
-	FString JsonCommand = CreateLoadFbxFileCommand(SourceData->GetFilename());
+	FString JsonCommand = CreateLoadFbxFileCommand(Filename);
 	int32 TaskIndex = Dispatcher->AddTask(JsonCommand);
 
 	//Blocking call until all tasks are executed
@@ -89,10 +91,14 @@ bool UInterchangeFbxTranslator::Translate(const UInterchangeSourceData* SourceDa
 		
 	UE::Interchange::ETaskState TaskState;
 	FString JsonResult;
-	TArray<FString> JSonMessages;
-	Dispatcher->GetTaskState(TaskIndex, TaskState, JsonResult, JSonMessages);
+	TArray<FString> JsonMessages;
+	Dispatcher->GetTaskState(TaskIndex, TaskState, JsonResult, JsonMessages);
 
-	//TODO: Parse the JSonMessage and add the message to the interchange not yet develop error messaging
+	// Parse the Json messages into UInterchangeResults
+	for (const FString& JsonMessage : JsonMessages)
+	{
+		AddMessage(UInterchangeResult::FromJson(JsonMessage));
+	}
 
 	if(TaskState != UE::Interchange::ETaskState::ProcessOk)
 	{
@@ -127,7 +133,7 @@ void UInterchangeFbxTranslator::ImportFinish()
 }
 
 
-TOptional<UE::Interchange::FImportImage> UInterchangeFbxTranslator::GetTexturePayloadData(const UInterchangeSourceData* SourceData, const FString& PayLoadKey) const
+TOptional<UE::Interchange::FImportImage> UInterchangeFbxTranslator::GetTexturePayloadData(const UInterchangeSourceData* InSourceData, const FString& PayLoadKey) const
 {
 	UInterchangeSourceData* PayloadSourceData = UInterchangeManager::GetInterchangeManager().CreateSourceData(PayLoadKey);
 	FGCObjectScopeGuard ScopedSourceData(PayloadSourceData);
@@ -146,10 +152,88 @@ TOptional<UE::Interchange::FImportImage> UInterchangeFbxTranslator::GetTexturePa
 	return TextureTranslator->GetTexturePayloadData(PayloadSourceData, PayLoadKey);
 }
 
-TOptional<UE::Interchange::FStaticMeshPayloadData> UInterchangeFbxTranslator::GetStaticMeshPayloadData(const FString& PayLoadKey) const
+TFuture<TOptional<UE::Interchange::FStaticMeshPayloadData>> UInterchangeFbxTranslator::GetStaticMeshPayloadData(const FString& PayLoadKey) const
 {
-	//Not implemented, currently we do not have any payload data for static meshes
-	return TOptional<UE::Interchange::FStaticMeshPayloadData>();
+	TSharedPtr<TPromise<TOptional<UE::Interchange::FStaticMeshPayloadData>>> Promise = MakeShared<TPromise<TOptional<UE::Interchange::FStaticMeshPayloadData>>>();
+
+	if (!Dispatcher.IsValid())
+	{
+		Promise->SetValue(TOptional<UE::Interchange::FStaticMeshPayloadData>());
+		return Promise->GetFuture();
+	}
+
+	// Create a json command to read the fbx file
+	FString JsonCommand = CreateFetchPayloadFbxCommand(PayLoadKey);
+	const int32 CreatedTaskIndex = Dispatcher->AddTask(JsonCommand, FInterchangeDispatcherTaskCompleted::CreateLambda([this, Promise](const int32 TaskIndex)
+	{
+		UE::Interchange::ETaskState TaskState;
+		FString JsonResult;
+		TArray<FString> JsonMessages;
+		Dispatcher->GetTaskState(TaskIndex, TaskState, JsonResult, JsonMessages);
+
+		// Parse the Json messages into UInterchangeResults
+		for (const FString& JsonMessage : JsonMessages)
+		{
+			AddMessage(UInterchangeResult::FromJson(JsonMessage));
+		}
+
+		if (TaskState != UE::Interchange::ETaskState::ProcessOk)
+		{
+			Promise->SetValue(TOptional<UE::Interchange::FStaticMeshPayloadData>());
+			return;
+		}
+
+		// Grab the result file and fill the BaseNodeContainer
+		UE::Interchange::FJsonFetchPayloadCmd::JsonResultParser ResultParser;
+		ResultParser.FromJson(JsonResult);
+		FString StaticMeshPayloadFilename = ResultParser.GetResultFilename();
+
+		if (!ensure(FPaths::FileExists(StaticMeshPayloadFilename)))
+		{
+			// TODO log an error saying the payload file does not exist even if the get payload command succeeded
+			Promise->SetValue(TOptional<UE::Interchange::FStaticMeshPayloadData>());
+			return;
+		}
+
+		UE::Interchange::FStaticMeshPayloadData StaticMeshPayload;
+		StaticMeshPayload.MeshDescription.Empty();
+
+		// All sub object should be gone with the reset
+		TArray64<uint8> Buffer;
+		FFileHelper::LoadFileToArray(Buffer, *StaticMeshPayloadFilename);
+		uint8* FileData = Buffer.GetData();
+		int64 FileDataSize = Buffer.Num();
+		if (FileDataSize < 1)
+		{
+			// Nothing to load from this file
+			Promise->SetValue(TOptional<UE::Interchange::FStaticMeshPayloadData>());
+			return;
+		}
+
+		// Buffer keeps the ownership of the data, the large memory reader is use to serialize the TMap
+		FLargeMemoryReader Ar(FileData, FileDataSize);
+		StaticMeshPayload.MeshDescription.Serialize(Ar);
+
+		// This is a static mesh and the payload should never contain skinned data
+		bool bFetchSkinnedData = false;
+		Ar << bFetchSkinnedData;
+		if (bFetchSkinnedData)
+		{
+			// TODO log an error saying the payload file is the wrong type for a static mesh
+			Promise->SetValue(TOptional<UE::Interchange::FStaticMeshPayloadData>());
+			return;
+		}
+
+		Promise->SetValue(MoveTemp(StaticMeshPayload));
+	}));
+
+	// The task was not added to the dispatcher
+	if (CreatedTaskIndex == INDEX_NONE)
+	{
+		Promise->SetValue(TOptional<UE::Interchange::FStaticMeshPayloadData>{});
+	}
+
+	return Promise->GetFuture();
 }
 
 TFuture<TOptional<UE::Interchange::FSkeletalMeshLodPayloadData>> UInterchangeFbxTranslator::GetSkeletalMeshLodPayloadData(const FString& PayLoadKey) const
@@ -167,10 +251,15 @@ TFuture<TOptional<UE::Interchange::FSkeletalMeshLodPayloadData>> UInterchangeFbx
 	{
 		UE::Interchange::ETaskState TaskState;
 		FString JsonResult;
-		TArray<FString> JSonMessages;
-		Dispatcher->GetTaskState(TaskIndex, TaskState, JsonResult, JSonMessages);
+		TArray<FString> JsonMessages;
+		Dispatcher->GetTaskState(TaskIndex, TaskState, JsonResult, JsonMessages);
+
+		// Parse the Json messages into UInterchangeResults
+		for (const FString& JsonMessage : JsonMessages)
+		{
+			AddMessage(UInterchangeResult::FromJson(JsonMessage));
+		}
 
-		//TODO: Parse the JSonMessage and add the message to the interchange not yet develop error messaging
 		if (TaskState != UE::Interchange::ETaskState::ProcessOk)
 		{
 			Promise->SetValue(TOptional<UE::Interchange::FSkeletalMeshLodPayloadData>{});
@@ -240,10 +329,14 @@ TFuture<TOptional<UE::Interchange::FSkeletalMeshBlendShapePayloadData>> UInterch
 	{
 		UE::Interchange::ETaskState TaskState;
 		FString JsonResult;
-		TArray<FString> JSonMessages;
-		Dispatcher->GetTaskState(TaskIndex, TaskState, JsonResult, JSonMessages);
+		TArray<FString> JsonMessages;
+		Dispatcher->GetTaskState(TaskIndex, TaskState, JsonResult, JsonMessages);
 
-		//TODO: Parse the JSonMessage and add the message to the interchange not yet develop error messaging
+		// Parse the Json messages into UInterchangeResults
+		for (const FString& JsonMessage : JsonMessages)
+		{
+			AddMessage(UInterchangeResult::FromJson(JsonMessage));
+		}
 
 		if (TaskState != UE::Interchange::ETaskState::ProcessOk)
 		{
diff --git a/Engine/Plugins/Runtime/Interchange/Source/Import/Private/InterchangeImportModule.cpp b/Engine/Plugins/Runtime/Interchange/Source/Import/Private/InterchangeImportModule.cpp
--- a/Engine/Plugins/Runtime/Interchange/Source/Import/Private/InterchangeImportModule.cpp
+++ b/Engine/Plugins/Runtime/Interchange/Source/Import/Private/InterchangeImportModule.cpp
@@ -9,6 +9,7 @@
 #include "Material/InterchangeMaterialFactory.h"
 #include "Mesh/InterchangeSkeletalMeshFactory.h"
 #include "Mesh/InterchangeSkeletonFactory.h"
+#include "Mesh/InterchangeStaticMeshFactory.h"
 #include "Misc/CoreDelegates.h"
 #include "Modules/ModuleManager.h"
 #include "Texture/InterchangeBMPTranslator.h"
@@ -68,6 +69,7 @@ void FInterchangeImportModule::StartupModule()
 		InterchangeManager.RegisterFactory(UInterchangeMaterialFactory::StaticClass());
 		InterchangeManager.RegisterFactory(UInterchangeSkeletonFactory::StaticClass());
 		InterchangeManager.RegisterFactory(UInterchangeSkeletalMeshFactory::StaticClass());
+		InterchangeManager.RegisterFactory(UInterchangeStaticMeshFactory::StaticClass());
 	};
 
 	if (GEngine)
diff --git a/Engine/Plugins/Runtime/Interchange/Source/Import/Private/Mesh/InterchangeStaticMeshFactory.cpp b/Engine/Plugins/Runtime/Interchange/Source/Import/Private/Mesh/InterchangeStaticMeshFactory.cpp
new file mode 100644
--- /dev/null
+++ b/Engine/Plugins/Runtime/Interchange/Source/Import/Private/Mesh/InterchangeStaticMeshFactory.cpp
@@ -0,0 +1,370 @@
+// Copyright Epic Games, Inc. All Rights Reserved. 
+#include "Mesh/InterchangeStaticMeshFactory.h"
+
+#include "Components.h"
+#include "Engine/StaticMesh.h"
+#include "InterchangeImportCommon.h"
+#include "InterchangeImportLog.h"
+#include "InterchangeMaterialFactoryNode.h"
+#include "InterchangeMaterialNode.h"
+#include "InterchangeMeshNode.h"
+#include "InterchangeSceneNode.h"
+#include "InterchangeStaticMeshLodDataNode.h"
+#include "InterchangeStaticMeshFactoryNode.h"
+#include "InterchangeSourceData.h"
+#include "InterchangeTranslatorBase.h"
+#include "Mesh/InterchangeStaticMeshPayload.h"
+#include "Mesh/InterchangeStaticMeshPayloadInterface.h"
+#include "Nodes/InterchangeBaseNode.h"
+#include "Nodes/InterchangeBaseNodeContainer.h"
+#include "StaticMeshAttributes.h"
+#include "StaticMeshOperations.h"
+
+#if WITH_EDITORONLY_DATA
+
+#include "EditorFramework/AssetImportData.h"
+
+#endif //WITH_EDITORONLY_DATA
+
+
+UClass* UInterchangeStaticMeshFactory::GetFactoryClass() const
+{
+	return UStaticMesh::StaticClass();
+}
+
+
+UObject* UInterchangeStaticMeshFactory::CreateEmptyAsset(const FCreateAssetParams& Arguments) const
+{
+#if !WITH_EDITOR || !WITH_EDITORONLY_DATA
+
+	UE_LOG(LogInterchangeImport, Error, TEXT("Cannot import StaticMesh asset in runtime, this is an editor only feature."));
+	return nullptr;
+
+#else
+	UStaticMesh* StaticMesh = nullptr;
+	if (!Arguments.AssetNode || !Arguments.AssetNode->GetAssetClass()->IsChildOf(GetFactoryClass()))
+	{
+		return nullptr;
+	}
+
+	const UInterchangeStaticMeshFactoryNode* StaticMeshFactoryNode = Cast<UInterchangeStaticMeshFactoryNode>(Arguments.AssetNode);
+	if (StaticMeshFactoryNode == nullptr)
+	{
+		return nullptr;
+	}
+
+	// create an asset if it doesn't exist
+	UObject* ExistingAsset = StaticFindObject(nullptr, Arguments.Parent, *Arguments.AssetName);
+
+	// create a new static mesh or overwrite existing asset, if possible
+	if (!ExistingAsset)
+	{
+		StaticMesh = NewObject<UStaticMesh>(Arguments.Parent, *Arguments.AssetName, RF_Public | RF_Standalone);
+	}
+	else if (ExistingAsset->GetClass()->IsChildOf(UStaticMesh::StaticClass()))
+	{
+		//This is a reimport, we are just re-updating the source data
+		StaticMesh = Cast<UStaticMesh>(ExistingAsset);
+	}
+	
+	if (!StaticMesh)
+	{
+		UE_LOG(LogInterchangeImport, Warning, TEXT("Could not create StaticMesh asset %s"), *Arguments.AssetName);
+		return nullptr;
+	}
+	
+	StaticMesh->PreEditChange(nullptr);	
+	return StaticMesh;
+#endif //else !WITH_EDITOR || !WITH_EDITORONLY_DATA
+}
+
+
+UObject* UInterchangeStaticMeshFactory::CreateAsset(const UInterchangeStaticMeshFactory::FCreateAssetParams& Arguments) const
+{
+#if !WITH_EDITOR || !WITH_EDITORONLY_DATA
+
+	UE_LOG(LogInterchangeImport, Error, TEXT("Cannot import static mesh asset in runtime, this is an editor only feature."));
+	return nullptr;
+
+#else
+
+	if (!Arguments.AssetNode || !Arguments.AssetNode->GetAssetClass()->IsChildOf(GetFactoryClass()))
+	{
+		return nullptr;
+	}
+
+	const UInterchangeStaticMeshFactoryNode* StaticMeshFactoryNode = Cast<UInterchangeStaticMeshFactoryNode>(Arguments.AssetNode);
+	if (StaticMeshFactoryNode == nullptr)
+	{
+		return nullptr;
+	}
+
+	const IInterchangeStaticMeshPayloadInterface* StaticMeshTranslatorPayloadInterface = Cast<IInterchangeStaticMeshPayloadInterface>(Arguments.Translator);
+	if (!StaticMeshTranslatorPayloadInterface)
+	{
+		UE_LOG(LogInterchangeImport, Error, TEXT("Cannot import static mesh, the translator does not implement the IInterchangeStaticMeshPayloadInterface."));
+		return nullptr;
+	}
+
+	const UClass* StaticMeshClass = StaticMeshFactoryNode->GetAssetClass();
+	check(StaticMeshClass && StaticMeshClass->IsChildOf(GetFactoryClass()));
+
+	// create an asset if it doesn't exist
+	UObject* ExistingAsset = StaticFindObject(nullptr, Arguments.Parent, *Arguments.AssetName);
+
+	UObject* StaticMeshObject = nullptr;
+
+	// create a new static mesh or overwrite existing asset, if possible
+	if (!ExistingAsset)
+	{
+		// NewObject is not thread safe, the asset registry directory watcher tick on the main thread can trig before we finish initializing the UObject and will crash
+		// The UObject should have been created by calling CreateEmptyAsset on the main thread.
+		check(IsInGameThread());
+		StaticMeshObject = NewObject<UObject>(Arguments.Parent, StaticMeshClass, *Arguments.AssetName, RF_Public | RF_Standalone);
+	}
+	else if(ExistingAsset->GetClass()->IsChildOf(StaticMeshClass))
+	{
+		//This is a reimport, we are just re-updating the source data
+		StaticMeshObject = ExistingAsset;
+	}
+
+	if (!StaticMeshObject)
+	{
+		UE_LOG(LogInterchangeImport, Error, TEXT("Could not create StaticMesh asset %s"), *Arguments.AssetName);
+		return nullptr;
+	}
+
+	if (StaticMeshObject)
+	{
+		// Currently static mesh re-import will not touch the static mesh at all
+		// TODO design a re-import process for the static mesh (expressions and input connections)
+		if (!Arguments.ReimportObject)
+		{
+			UStaticMesh* StaticMesh = Cast<UStaticMesh>(StaticMeshObject);
+			if (!ensure(StaticMesh))
+			{
+				UE_LOG(LogInterchangeImport, Error, TEXT("Could not create StaticMesh asset %s"), *Arguments.AssetName);
+				return nullptr;
+			}
+			
+			int32 LodCount = StaticMeshFactoryNode->GetLodDataCount();
+			StaticMesh->SetNumSourceModels(LodCount);
+
+			TArray<FString> LodDataUniqueIds;
+			StaticMeshFactoryNode->GetLodDataUniqueIds(LodDataUniqueIds);
+			ensure(LodDataUniqueIds.Num() == LodCount);
+
+			int32 CurrentLodIndex = 0;
+			for (int32 LodIndex = 0; LodIndex < LodCount; ++LodIndex)
+			{
+				FString LodUniqueId = LodDataUniqueIds[LodIndex];
+				const UInterchangeStaticMeshLodDataNode* LodDataNode = Cast<UInterchangeStaticMeshLodDataNode>(Arguments.NodeContainer->GetNode(LodUniqueId));
+				if (!LodDataNode)
+				{
+					UE_LOG(LogInterchangeImport, Warning, TEXT("Invalid LOD when importing StaticMesh asset %s"), *Arguments.AssetName);
+					continue;
+				}
+
+				// Get the mesh node context for each MeshUids
+				struct FMeshNodeContext
+				{
+					const UInterchangeMeshNode* MeshNode = nullptr;
+					const UInterchangeSceneNode* SceneNode = nullptr;
+					TOptional<FTransform> SceneGlobalTransform;
+					FString TranslatorPayloadKey;
+				};
+				TArray<FMeshNodeContext> MeshReferences;
+
+				// Scope to query the mesh node
+				{
+					TArray<FString> MeshUids;
+					LodDataNode->GetMeshUids(MeshUids);
+					MeshReferences.Reserve(MeshUids.Num());
+
+					for (const FString& MeshUid : MeshUids)
+					{
+						FMeshNodeContext MeshReference;
+						MeshReference.MeshNode = Cast<UInterchangeMeshNode>(Arguments.NodeContainer->GetNode(MeshUid));
+						if (!MeshReference.MeshNode)
+						{
+							// The reference is a scene node and we need to bake the geometry
+							MeshReference.SceneNode = Cast<UInterchangeSceneNode>(Arguments.NodeContainer->GetNode(MeshUid));
+							if (!ensure(MeshReference.SceneNode != nullptr))
+							{
+								UE_LOG(LogInterchangeImport, Warning, TEXT("Invalid LOD mesh reference when importing StaticMesh asset %s"), *Arguments.AssetName);
+								continue;
+							}
+
+							FString MeshDependencyUid;
+							MeshReference.SceneNode->GetCustomMeshDependencyUid(MeshDependencyUid);
+							MeshReference.MeshNode = Cast<UInterchangeMeshNode>(Arguments.NodeContainer->GetNode(MeshDependencyUid));
+
+							// Cache the scene node global matrix, we will use this matrix to bake the vertices
+							FTransform SceneNodeGlobalTransform;
+							if (MeshReference.SceneNode->GetCustomGlobalTransform(SceneNodeGlobalTransform))
+							{
+								MeshReference.SceneGlobalTransform = SceneNodeGlobalTransform;
+							}
+						}
+
+						if (!ensure(MeshReference.MeshNode != nullptr))
+						{
+							UE_LOG(LogInterchangeImport, Warning, TEXT("Invalid LOD mesh reference when importing StaticMesh asset %s"), *Arguments.AssetName);
+							continue;
+						}
+
+						TOptional<FString> MeshPayloadKey = MeshReference.MeshNode->GetPayLoadKey();
+						if (MeshPayloadKey.IsSet())
+						{
+							MeshReference.TranslatorPayloadKey = MeshPayloadKey.GetValue();
+						}
+						else
+						{
+							UE_LOG(LogInterchangeImport, Warning, TEXT("Empty LOD mesh reference payload when importing StaticMesh asset %s"), *Arguments.AssetName);
+							continue;
+						}
+
+						MeshReferences.Add(MeshReference);
+					}
+				}
+
+				// Add the lod mesh data to the static mesh
+				FMeshDescription* LodMeshDescription = StaticMesh->CreateMeshDescription(CurrentLodIndex);
+
+				FStaticMeshOperations::FAppendSettings AppendSettings;
+				for (int32 ChannelIdx = 0; ChannelIdx < FStaticMeshOperations::FAppendSettings::MAX_NUM_UV_CHANNELS; ++ChannelIdx)
+				{
+					AppendSettings.bMergeUVChannels[ChannelIdx] = true;
+				}
+
+				// Fetch the payloads for all mesh references in parallel
+				TMap<FString, TFuture<TOptional<UE::Interchange::FStaticMeshPayloadData>>> MeshPayloads;
+				for (const FMeshNodeContext& MeshNodeContext : MeshReferences)
+				{
+					const FString& PayloadKey = MeshNodeContext.TranslatorPayloadKey;
+					MeshPayloads.Add(PayloadKey, StaticMeshTranslatorPayloadInterface->GetStaticMeshPayloadData(PayloadKey));
+				}
+				
+				// Fill the lod mesh description using all combined mesh parts
+				for (const FMeshNodeContext& MeshNodeContext : MeshReferences)
+				{
+					TOptional<UE::Interchange::FStaticMeshPayloadData> LodMeshPayload = MeshPayloads.FindChecked(MeshNodeContext.TranslatorPayloadKey).Get();
+					if (!LodMeshPayload.IsSet())
+					{
+						UE_LOG(LogInterchangeImport, Warning, TEXT("Invalid static mesh payload key [%s] StaticMesh asset %s"), *MeshNodeContext.TranslatorPayloadKey, *Arguments.AssetName);
+						continue;
+					}
+
+					// Bake the payload, with the provided transform
+					if (MeshNodeContext.SceneGlobalTransform.IsSet())
+					{
+						AppendSettings.MeshTransform = MeshNodeContext.SceneGlobalTransform;
+					}
+					else
+					{
+						AppendSettings.MeshTransform.Reset();
+					}
+
+					FStaticMeshOperations::AppendMeshDescription(LodMeshPayload->MeshDescription, *LodMeshDescription, AppendSettings);
+				}
+
+				StaticMesh->CommitMeshDescription(CurrentLodIndex);
+
+				// Handle materials
+				// @TODO: move this outside of the LOD loop?
+				TArray<FString> FactoryDependencies;
+				StaticMeshFactoryNode->GetFactoryDependencies(FactoryDependencies);
+
+				for (int32 DependencyIndex = 0; DependencyIndex < FactoryDependencies.Num(); ++DependencyIndex)
+				{
+					const UInterchangeMaterialFactoryNode* MaterialFactoryNode = Cast<UInterchangeMaterialFactoryNode>(Arguments.NodeContainer->GetNode(FactoryDependencies[DependencyIndex]));
+					if (!MaterialFactoryNode || !MaterialFactoryNode->ReferenceObject.IsValid())
+					{
+						continue;
+					}
+
+					FName MaterialSlotName = *MaterialFactoryNode->GetDisplayLabel();
+					UMaterialInterface* MaterialInterface = Cast<UMaterialInterface>(MaterialFactoryNode->ReferenceObject.ResolveObject());
+					if (!MaterialInterface)
+					{
+						MaterialInterface = UMaterial::GetDefaultMaterial(MD_Surface);
+					}
+
+					int32 MaterialSlotIndex = StaticMesh->GetMaterialIndexFromImportedMaterialSlotName(MaterialSlotName);
+					if (MaterialSlotIndex == INDEX_NONE)
+					{
+#if WITH_EDITORONLY_DATA
+						StaticMesh->GetStaticMaterials().Emplace(MaterialInterface, MaterialSlotName, MaterialSlotName);
+#else
+						StaticMesh->GetStaticMaterials().Emplace(MaterialInterface, MaterialSlotName);
+#endif
+					}
+				}
+
+				FStaticMeshConstAttributes StaticMeshAttributes(*LodMeshDescription);
+				TPolygonGroupAttributesRef<const FName> SlotNames = StaticMeshAttributes.GetPolygonGroupMaterialSlotNames();
+				int32 SectionIndex = 0;
+				for (FPolygonGroupID PolygonGroupID : LodMeshDescription->PolygonGroups().GetElementIDs())
+				{
+					int32 MaterialSlotIndex = StaticMesh->GetMaterialIndexFromImportedMaterialSlotName(SlotNames[PolygonGroupID]);
+					check(MaterialSlotIndex != INDEX_NONE);
+
+					FMeshSectionInfo Info = StaticMesh->GetSectionInfoMap().Get(CurrentLodIndex, SectionIndex);
+					Info.MaterialIndex = MaterialSlotIndex;
+					StaticMesh->GetSectionInfoMap().Remove(CurrentLodIndex, SectionIndex);
+					StaticMesh->GetSectionInfoMap().Set(CurrentLodIndex, SectionIndex, Info);
+
+					SectionIndex++;
+				}
+
+				// Set static mesh build settings
+				FStaticMeshSourceModel& SrcModel = StaticMesh->GetSourceModel(CurrentLodIndex);
+				SrcModel.ReductionSettings.MaxDeviation = 0.0f;
+				SrcModel.ReductionSettings.PercentTriangles = 1.0f;
+				SrcModel.ReductionSettings.PercentVertices = 1.0f;
+
+				CurrentLodIndex++;
+			}
+
+			// Apply all SkeletalMeshFactoryNode custom attributes to the material asset
+			StaticMeshFactoryNode->ApplyAllCustomAttributeToAsset(StaticMesh);
+		}
+		
+		// Getting the file Hash will cache it into the source data
+		Arguments.SourceData->GetFileContentHash();
+
+		// The interchange completion task (called in the GameThread after the factories pass), will call PostEditChange
+		// which will trigger another asynchronous system that will build all static meshes in parallel
+	}
+	else
+	{
+		// The static mesh is not a UStaticMesh
+		StaticMeshObject->RemoveFromRoot();
+		StaticMeshObject->MarkPendingKill();
+	}
+	return StaticMeshObject;
+
+#endif //else !WITH_EDITOR || !WITH_EDITORONLY_DATA
+}
+
+
+/* This function is call in the completion task on the main thread, use it to call main thread post creation step for your assets */
+void UInterchangeStaticMeshFactory::PostImportGameThreadCallback(const FPostImportGameThreadCallbackParams& Arguments) const
+{
+	check(IsInGameThread());
+	Super::PostImportGameThreadCallback(Arguments);
+
+	// TODO: make sure this works at runtime
+#if WITH_EDITORONLY_DATA
+	if (ensure(Arguments.ImportedObject && Arguments.SourceData))
+	{
+		// We must call the Update of the asset source file in the main thread because UAssetImportData::Update execute some delegate we do not control
+		UStaticMesh* StaticMesh = CastChecked<UStaticMesh>(Arguments.ImportedObject);
+
+		UAssetImportData* ImportDataPtr = StaticMesh->GetAssetImportData();
+		UE::Interchange::FFactoryCommon::FUpdateImportAssetDataParameters UpdateImportAssetDataParameters(StaticMesh, ImportDataPtr, Arguments.SourceData, Arguments.NodeUniqueID, Arguments.NodeContainer, Arguments.Pipelines);
+		ImportDataPtr = UE::Interchange::FFactoryCommon::UpdateImportAssetData(UpdateImportAssetDataParameters);
+		StaticMesh->SetAssetImportData(ImportDataPtr);
+	}
+#endif
+}
\ No newline at end of file
diff --git a/Engine/Plugins/Runtime/Interchange/Source/Import/Private/Texture/InterchangeBMPTranslator.cpp b/Engine/Plugins/Runtime/Interchange/Source/Import/Private/Texture/InterchangeBMPTranslator.cpp
--- a/Engine/Plugins/Runtime/Interchange/Source/Import/Private/Texture/InterchangeBMPTranslator.cpp
+++ b/Engine/Plugins/Runtime/Interchange/Source/Import/Private/Texture/InterchangeBMPTranslator.cpp
@@ -24,21 +24,23 @@ bool UInterchangeBMPTranslator::CanImportSourceData(const UInterchangeSourceData
 	return BMPExtension.StartsWith(Extension);
 }
 
-bool UInterchangeBMPTranslator::Translate(const UInterchangeSourceData* SourceData, UInterchangeBaseNodeContainer& BaseNodeContainer) const
+bool UInterchangeBMPTranslator::Translate(UInterchangeBaseNodeContainer& BaseNodeContainer) const
 {
-	return UE::Interchange::FTextureTranslatorUtilities::Generic2DTextureTranslate(SourceData, BaseNodeContainer);
+	return UE::Interchange::FTextureTranslatorUtilities::Generic2DTextureTranslate(GetSourceData(), BaseNodeContainer);
 }
 
-TOptional<UE::Interchange::FImportImage> UInterchangeBMPTranslator::GetTexturePayloadData(const UInterchangeSourceData* SourceData, const FString& PayLoadKey) const
+TOptional<UE::Interchange::FImportImage> UInterchangeBMPTranslator::GetTexturePayloadData(const UInterchangeSourceData* PayloadSourceData, const FString& PayLoadKey) const
 {
-	if (!SourceData)
+	check(PayloadSourceData == GetSourceData());
+
+	if (!GetSourceData())
 	{
 		UE_LOG(LogInterchangeImport, Error, TEXT("Failed to import BMP, bad source data."));
 		return TOptional<UE::Interchange::FImportImage>();
 	}
 
 	TArray64<uint8> SourceDataBuffer;
-	FString Filename = SourceData->GetFilename();
+	FString Filename = GetSourceData()->GetFilename();
 
 	//Make sure the key fit the filename, The key should always be valid
 	if (!Filename.Equals(PayLoadKey))
diff --git a/Engine/Plugins/Runtime/Interchange/Source/Import/Private/Texture/InterchangeDDSTranslator.cpp b/Engine/Plugins/Runtime/Interchange/Source/Import/Private/Texture/InterchangeDDSTranslator.cpp
--- a/Engine/Plugins/Runtime/Interchange/Source/Import/Private/Texture/InterchangeDDSTranslator.cpp
+++ b/Engine/Plugins/Runtime/Interchange/Source/Import/Private/Texture/InterchangeDDSTranslator.cpp
@@ -79,12 +79,12 @@ bool UInterchangeDDSTranslator::CanImportSourceData(const UInterchangeSourceData
 	return DDSLoadHelper.IsValid2DTexture() || DDSLoadHelper.IsValidCubemapTexture() || DDSLoadHelper.IsValidArrayTexture();
 }
 
-bool UInterchangeDDSTranslator::Translate(const UInterchangeSourceData* SourceData, UInterchangeBaseNodeContainer& BaseNodeContainer) const
+bool UInterchangeDDSTranslator::Translate(UInterchangeBaseNodeContainer& BaseNodeContainer) const
 {
 	/*
 	 * DDS file can also be a cube map so we have to open the file and see if its a valid 2D texture.
 	 */
-	FString Filename = SourceData->GetFilename();
+	FString Filename = GetSourceData()->GetFilename();
 	if (!FPaths::FileExists(Filename))
 	{
 		return false;
@@ -100,31 +100,33 @@ bool UInterchangeDDSTranslator::Translate(const UInterchangeSourceData* SourceDa
 	FDDSLoadHelper DDSLoadHelper(HeaderDataBuffer.GetData(), HeaderDataBuffer.Num());
 	if (DDSLoadHelper.IsValid2DTexture())
 	{
-		return UE::Interchange::FTextureTranslatorUtilities::Generic2DTextureTranslate(SourceData, BaseNodeContainer);
+		return UE::Interchange::FTextureTranslatorUtilities::Generic2DTextureTranslate(GetSourceData(), BaseNodeContainer);
 	}
 
 	if (DDSLoadHelper.IsValidCubemapTexture())
 	{
-		return UE::Interchange::FTextureTranslatorUtilities::GenericTextureCubeTranslate(SourceData, BaseNodeContainer);
+		return UE::Interchange::FTextureTranslatorUtilities::GenericTextureCubeTranslate(GetSourceData(), BaseNodeContainer);
 	}
 
 	if (DDSLoadHelper.IsValidArrayTexture())
 	{
-		return UE::Interchange::FTextureTranslatorUtilities::GenericTexture2DArrayTranslate(SourceData, BaseNodeContainer);
+		return UE::Interchange::FTextureTranslatorUtilities::GenericTexture2DArrayTranslate(GetSourceData(), BaseNodeContainer);
 	}
 
 	return false;
 }
 
-TOptional<UE::Interchange::FImportImage> UInterchangeDDSTranslator::GetTexturePayloadData(const UInterchangeSourceData* SourceData, const FString& PayLoadKey) const
+TOptional<UE::Interchange::FImportImage> UInterchangeDDSTranslator::GetTexturePayloadData(const UInterchangeSourceData* PayloadSourceData, const FString& PayLoadKey) const
 {
-	if (!SourceData)
+	check(PayloadSourceData == GetSourceData());
+
+	if (!GetSourceData())
 	{
 		UE_LOG(LogInterchangeImport, Error, TEXT("Failed to import DDS, bad source data."));
 		return TOptional<UE::Interchange::FImportImage>();
 	}
 
-	FString Filename = SourceData->GetFilename();
+	FString Filename = GetSourceData()->GetFilename();
 	
 	//Make sure the key fit the filename, The key should always be valid
 	if (!Filename.Equals(PayLoadKey))
@@ -206,15 +208,17 @@ TOptional<UE::Interchange::FImportImage> UInterchangeDDSTranslator::GetTexturePa
 	return PayloadData;
 }
 
-TOptional<UE::Interchange::FImportSlicedImage> UInterchangeDDSTranslator::GetSlicedTexturePayloadData(const UInterchangeSourceData* SourceData, const FString& PayLoadKey) const
+TOptional<UE::Interchange::FImportSlicedImage> UInterchangeDDSTranslator::GetSlicedTexturePayloadData(const UInterchangeSourceData* PayloadSourceData, const FString& PayLoadKey) const
 {
-	if (!SourceData)
+	check(PayloadSourceData == GetSourceData());
+
+	if (!GetSourceData())
 	{
 		UE_LOG(LogInterchangeImport, Error, TEXT("Failed to import DDS, bad source data."));
 		return TOptional<UE::Interchange::FImportSlicedImage>();
 	}
 
-	FString Filename = SourceData->GetFilename();
+	FString Filename = GetSourceData()->GetFilename();
 	
 	//Make sure the key fit the filename, The key should always be valid
 	if (!Filename.Equals(PayLoadKey))
diff --git a/Engine/Plugins/Runtime/Interchange/Source/Import/Private/Texture/InterchangeEXRTranslator.cpp b/Engine/Plugins/Runtime/Interchange/Source/Import/Private/Texture/InterchangeEXRTranslator.cpp
--- a/Engine/Plugins/Runtime/Interchange/Source/Import/Private/Texture/InterchangeEXRTranslator.cpp
+++ b/Engine/Plugins/Runtime/Interchange/Source/Import/Private/Texture/InterchangeEXRTranslator.cpp
@@ -22,21 +22,23 @@ bool UInterchangeEXRTranslator::CanImportSourceData(const UInterchangeSourceData
 	return EXRExtension.StartsWith(Extension);
 }
 
-bool UInterchangeEXRTranslator::Translate(const UInterchangeSourceData* SourceData, UInterchangeBaseNodeContainer& BaseNodeContainer) const
+bool UInterchangeEXRTranslator::Translate(UInterchangeBaseNodeContainer& BaseNodeContainer) const
 {
-	return UE::Interchange::FTextureTranslatorUtilities::Generic2DTextureTranslate(SourceData, BaseNodeContainer);
+	return UE::Interchange::FTextureTranslatorUtilities::Generic2DTextureTranslate(GetSourceData(), BaseNodeContainer);
 }
 
-TOptional<UE::Interchange::FImportImage> UInterchangeEXRTranslator::GetTexturePayloadData(const UInterchangeSourceData* SourceData, const FString& PayLoadKey) const
+TOptional<UE::Interchange::FImportImage> UInterchangeEXRTranslator::GetTexturePayloadData(const UInterchangeSourceData* PayloadSourceData, const FString& PayLoadKey) const
 {
-	if (!SourceData)
+	check(PayloadSourceData == GetSourceData());
+
+	if (!GetSourceData())
 	{
 		UE_LOG(LogInterchangeImport, Error, TEXT("Failed to import EXR, bad source data."));
 		return TOptional<UE::Interchange::FImportImage>();
 	}
 
 	TArray64<uint8> SourceDataBuffer;
-	FString Filename = SourceData->GetFilename();
+	FString Filename = GetSourceData()->GetFilename();
 	
 	//Make sure the key fit the filename, The key should always be valid
 	if (!Filename.Equals(PayLoadKey))
diff --git a/Engine/Plugins/Runtime/Interchange/Source/Import/Private/Texture/InterchangeHDRTranslator.cpp b/Engine/Plugins/Runtime/Interchange/Source/Import/Private/Texture/InterchangeHDRTranslator.cpp
--- a/Engine/Plugins/Runtime/Interchange/Source/Import/Private/Texture/InterchangeHDRTranslator.cpp
+++ b/Engine/Plugins/Runtime/Interchange/Source/Import/Private/Texture/InterchangeHDRTranslator.cpp
@@ -23,21 +23,23 @@ bool UInterchangeHDRTranslator::CanImportSourceData(const UInterchangeSourceData
 	return HDRExtension.StartsWith(Extension);
 }
 
-bool UInterchangeHDRTranslator::Translate(const UInterchangeSourceData* SourceData, UInterchangeBaseNodeContainer& BaseNodeContainer) const
+bool UInterchangeHDRTranslator::Translate(UInterchangeBaseNodeContainer& BaseNodeContainer) const
 {
-	return UE::Interchange::FTextureTranslatorUtilities::Generic2DTextureTranslate(SourceData, BaseNodeContainer);
+	return UE::Interchange::FTextureTranslatorUtilities::Generic2DTextureTranslate(GetSourceData(), BaseNodeContainer);
 }
 
-TOptional<UE::Interchange::FImportImage> UInterchangeHDRTranslator::GetTexturePayloadData(const UInterchangeSourceData* SourceData, const FString& PayLoadKey) const
+TOptional<UE::Interchange::FImportImage> UInterchangeHDRTranslator::GetTexturePayloadData(const UInterchangeSourceData* PayloadSourceData, const FString& PayLoadKey) const
 {
-	if (!SourceData)
+	check(PayloadSourceData == GetSourceData());
+
+	if (!GetSourceData())
 	{
 		UE_LOG(LogInterchangeImport, Error, TEXT("Failed to import HDR, bad source data."));
 		return TOptional<UE::Interchange::FImportImage>();
 	}
 
 	TArray64<uint8> SourceDataBuffer;
-	FString Filename = SourceData->GetFilename();
+	FString Filename = GetSourceData()->GetFilename();
 
 	//Make sure the key fit the filename, The key should always be valid
 	if (!Filename.Equals(PayLoadKey))
diff --git a/Engine/Plugins/Runtime/Interchange/Source/Import/Private/Texture/InterchangeIESTranslator.cpp b/Engine/Plugins/Runtime/Interchange/Source/Import/Private/Texture/InterchangeIESTranslator.cpp
--- a/Engine/Plugins/Runtime/Interchange/Source/Import/Private/Texture/InterchangeIESTranslator.cpp
+++ b/Engine/Plugins/Runtime/Interchange/Source/Import/Private/Texture/InterchangeIESTranslator.cpp
@@ -18,21 +18,23 @@ bool UInterchangeIESTranslator::CanImportSourceData(const UInterchangeSourceData
 	return IESExtension.StartsWith(Extension);
 }
 
-bool UInterchangeIESTranslator::Translate(const UInterchangeSourceData* SourceData, UInterchangeBaseNodeContainer& BaseNodeContainer) const
+bool UInterchangeIESTranslator::Translate(UInterchangeBaseNodeContainer& BaseNodeContainer) const
 {
-	return UE::Interchange::FTextureTranslatorUtilities::GenericTextureLightProfileTranslate(SourceData, BaseNodeContainer);
+	return UE::Interchange::FTextureTranslatorUtilities::GenericTextureLightProfileTranslate(GetSourceData(), BaseNodeContainer);
 }
 
-TOptional<UE::Interchange::FImportLightProfile> UInterchangeIESTranslator::GetLightProfilePayloadData(const UInterchangeSourceData* SourceData, const FString& PayLoadKey) const
+TOptional<UE::Interchange::FImportLightProfile> UInterchangeIESTranslator::GetLightProfilePayloadData(const UInterchangeSourceData* PayloadSourceData, const FString& PayLoadKey) const
 {
-	if (!SourceData)
+	check(GetSourceData() == PayloadSourceData);
+
+	if (!GetSourceData())
 	{
 		UE_LOG(LogInterchangeImport, Error, TEXT("Failed to import IES, bad source data."));
 		return {};
 	}
 
 	TArray64<uint8> SourceDataBuffer;
-	FString Filename = SourceData->GetFilename();
+	FString Filename = GetSourceData()->GetFilename();
 
 	//Make sure the key fit the filename, The key should always be valid
 	if (!Filename.Equals(PayLoadKey))
diff --git a/Engine/Plugins/Runtime/Interchange/Source/Import/Private/Texture/InterchangeJPGTranslator.cpp b/Engine/Plugins/Runtime/Interchange/Source/Import/Private/Texture/InterchangeJPGTranslator.cpp
--- a/Engine/Plugins/Runtime/Interchange/Source/Import/Private/Texture/InterchangeJPGTranslator.cpp
+++ b/Engine/Plugins/Runtime/Interchange/Source/Import/Private/Texture/InterchangeJPGTranslator.cpp
@@ -22,21 +22,23 @@ bool UInterchangeJPGTranslator::CanImportSourceData(const UInterchangeSourceData
 	return JPGExtension.StartsWith(Extension);
 }
 
-bool UInterchangeJPGTranslator::Translate(const UInterchangeSourceData* SourceData, UInterchangeBaseNodeContainer& BaseNodeContainer) const
+bool UInterchangeJPGTranslator::Translate(UInterchangeBaseNodeContainer& BaseNodeContainer) const
 {
-	return UE::Interchange::FTextureTranslatorUtilities::Generic2DTextureTranslate(SourceData, BaseNodeContainer);
+	return UE::Interchange::FTextureTranslatorUtilities::Generic2DTextureTranslate(GetSourceData(), BaseNodeContainer);
 }
 
-TOptional<UE::Interchange::FImportImage> UInterchangeJPGTranslator::GetTexturePayloadData(const UInterchangeSourceData* SourceData, const FString& PayLoadKey) const
+TOptional<UE::Interchange::FImportImage> UInterchangeJPGTranslator::GetTexturePayloadData(const UInterchangeSourceData* PayloadSourceData, const FString& PayLoadKey) const
 {
-	if (!SourceData)
+	check(PayloadSourceData == GetSourceData());
+
+	if (!GetSourceData())
 	{
 		UE_LOG(LogInterchangeImport, Error, TEXT("Failed to import JPEG, bad source data."));
 		return TOptional<UE::Interchange::FImportImage>();
 	}
 
 	TArray64<uint8> SourceDataBuffer;
-	FString Filename = SourceData->GetFilename();
+	FString Filename = GetSourceData()->GetFilename();
 	
 	//Make sure the key fit the filename, The key should always be valid
 	if (!Filename.Equals(PayLoadKey))
diff --git a/Engine/Plugins/Runtime/Interchange/Source/Import/Private/Texture/InterchangePCXTranslator.cpp b/Engine/Plugins/Runtime/Interchange/Source/Import/Private/Texture/InterchangePCXTranslator.cpp
--- a/Engine/Plugins/Runtime/Interchange/Source/Import/Private/Texture/InterchangePCXTranslator.cpp
+++ b/Engine/Plugins/Runtime/Interchange/Source/Import/Private/Texture/InterchangePCXTranslator.cpp
@@ -73,21 +73,23 @@ bool UInterchangePCXTranslator::CanImportSourceData(const UInterchangeSourceData
 	return PCXExtension.StartsWith(Extension);
 }
 
-bool UInterchangePCXTranslator::Translate(const UInterchangeSourceData* SourceData, UInterchangeBaseNodeContainer& BaseNodeContainer) const
+bool UInterchangePCXTranslator::Translate(UInterchangeBaseNodeContainer& BaseNodeContainer) const
 {
-	return UE::Interchange::FTextureTranslatorUtilities::Generic2DTextureTranslate(SourceData, BaseNodeContainer);
+	return UE::Interchange::FTextureTranslatorUtilities::Generic2DTextureTranslate(GetSourceData(), BaseNodeContainer);
 }
 
-TOptional<UE::Interchange::FImportImage> UInterchangePCXTranslator::GetTexturePayloadData(const UInterchangeSourceData* SourceData, const FString& PayLoadKey) const
+TOptional<UE::Interchange::FImportImage> UInterchangePCXTranslator::GetTexturePayloadData(const UInterchangeSourceData* PayloadSourceData, const FString& PayLoadKey) const
 {
-	if (!SourceData)
+	check(PayloadSourceData == GetSourceData());
+
+	if (!GetSourceData())
 	{
 		UE_LOG(LogInterchangeImport, Error, TEXT("Failed to import PCX, bad source data."));
 		return TOptional<UE::Interchange::FImportImage>();
 	}
 
 	TArray64<uint8> SourceDataBuffer;
-	FString Filename = SourceData->GetFilename();
+	FString Filename = GetSourceData()->GetFilename();
 	
 	//Make sure the key fit the filename, The key should always be valid
 	if (!Filename.Equals(PayLoadKey))
diff --git a/Engine/Plugins/Runtime/Interchange/Source/Import/Private/Texture/InterchangePNGTranslator.cpp b/Engine/Plugins/Runtime/Interchange/Source/Import/Private/Texture/InterchangePNGTranslator.cpp
--- a/Engine/Plugins/Runtime/Interchange/Source/Import/Private/Texture/InterchangePNGTranslator.cpp
+++ b/Engine/Plugins/Runtime/Interchange/Source/Import/Private/Texture/InterchangePNGTranslator.cpp
@@ -194,21 +194,23 @@ bool UInterchangePNGTranslator::CanImportSourceData(const UInterchangeSourceData
 	return PNGExtension.StartsWith(Extension);
 }
 
-bool UInterchangePNGTranslator::Translate(const UInterchangeSourceData* SourceData, UInterchangeBaseNodeContainer& BaseNodeContainer) const
+bool UInterchangePNGTranslator::Translate(UInterchangeBaseNodeContainer& BaseNodeContainer) const
 {
-	return UE::Interchange::FTextureTranslatorUtilities::Generic2DTextureTranslate(SourceData, BaseNodeContainer);
+	return UE::Interchange::FTextureTranslatorUtilities::Generic2DTextureTranslate(GetSourceData(), BaseNodeContainer);
 }
 
-TOptional<UE::Interchange::FImportImage> UInterchangePNGTranslator::GetTexturePayloadData(const UInterchangeSourceData* SourceData, const FString& PayLoadKey) const
+TOptional<UE::Interchange::FImportImage> UInterchangePNGTranslator::GetTexturePayloadData(const UInterchangeSourceData* PayloadSourceData, const FString& PayLoadKey) const
 {
-	if (!SourceData)
+	check(PayloadSourceData == GetSourceData());
+
+	if (!GetSourceData())
 	{
 		UE_LOG(LogInterchangeImport, Error, TEXT("Failed to import PNG, bad source data."));
 		return TOptional<UE::Interchange::FImportImage>();
 	}
 
 	TArray64<uint8> SourceDataBuffer;
-	FString Filename = SourceData->GetFilename();
+	FString Filename = GetSourceData()->GetFilename();
 	
 	//Make sure the key fit the filename, The key should always be valid
 	if (!Filename.Equals(PayLoadKey))
diff --git a/Engine/Plugins/Runtime/Interchange/Source/Import/Private/Texture/InterchangePSDTranslator.cpp b/Engine/Plugins/Runtime/Interchange/Source/Import/Private/Texture/InterchangePSDTranslator.cpp
--- a/Engine/Plugins/Runtime/Interchange/Source/Import/Private/Texture/InterchangePSDTranslator.cpp
+++ b/Engine/Plugins/Runtime/Interchange/Source/Import/Private/Texture/InterchangePSDTranslator.cpp
@@ -375,21 +375,23 @@ bool UInterchangePSDTranslator::CanImportSourceData(const UInterchangeSourceData
 	return PNGExtension.StartsWith(Extension);
 }
 
-bool UInterchangePSDTranslator::Translate(const UInterchangeSourceData* SourceData, UInterchangeBaseNodeContainer& BaseNodeContainer) const
+bool UInterchangePSDTranslator::Translate(UInterchangeBaseNodeContainer& BaseNodeContainer) const
 {
-	return UE::Interchange::FTextureTranslatorUtilities::Generic2DTextureTranslate(SourceData, BaseNodeContainer);
+	return UE::Interchange::FTextureTranslatorUtilities::Generic2DTextureTranslate(GetSourceData(), BaseNodeContainer);
 }
 
-TOptional<UE::Interchange::FImportImage> UInterchangePSDTranslator::GetTexturePayloadData(const UInterchangeSourceData* SourceData, const FString& PayLoadKey) const
+TOptional<UE::Interchange::FImportImage> UInterchangePSDTranslator::GetTexturePayloadData(const UInterchangeSourceData* PayloadSourceData, const FString& PayLoadKey) const
 {
-	if (!SourceData)
+	check(PayloadSourceData == GetSourceData());
+
+	if (!GetSourceData())
 	{
 		UE_LOG(LogInterchangeImport, Error, TEXT("Failed to import PSD, bad source data."));
 		return TOptional<UE::Interchange::FImportImage>();
 	}
 
 	TArray64<uint8> SourceDataBuffer;
-	FString Filename = SourceData->GetFilename();
+	FString Filename = GetSourceData()->GetFilename();
 
 	//Make sure the key fit the filename, The key should always be valid
 	if (!Filename.Equals(PayLoadKey))
diff --git a/Engine/Plugins/Runtime/Interchange/Source/Import/Private/Texture/InterchangeTGATranslator.cpp b/Engine/Plugins/Runtime/Interchange/Source/Import/Private/Texture/InterchangeTGATranslator.cpp
--- a/Engine/Plugins/Runtime/Interchange/Source/Import/Private/Texture/InterchangeTGATranslator.cpp
+++ b/Engine/Plugins/Runtime/Interchange/Source/Import/Private/Texture/InterchangeTGATranslator.cpp
@@ -79,21 +79,23 @@ bool UInterchangeTGATranslator::CanImportSourceData(const UInterchangeSourceData
 	return TGAExtension.StartsWith(Extension);
 }
 
-bool UInterchangeTGATranslator::Translate(const UInterchangeSourceData* SourceData, UInterchangeBaseNodeContainer& BaseNodeContainer) const
+bool UInterchangeTGATranslator::Translate(UInterchangeBaseNodeContainer& BaseNodeContainer) const
 {
-	return UE::Interchange::FTextureTranslatorUtilities::Generic2DTextureTranslate(SourceData, BaseNodeContainer);
+	return UE::Interchange::FTextureTranslatorUtilities::Generic2DTextureTranslate(GetSourceData(), BaseNodeContainer);
 }
 
-TOptional<UE::Interchange::FImportImage> UInterchangeTGATranslator::GetTexturePayloadData(const UInterchangeSourceData* SourceData, const FString& PayLoadKey) const
+TOptional<UE::Interchange::FImportImage> UInterchangeTGATranslator::GetTexturePayloadData(const UInterchangeSourceData* PayloadSourceData, const FString& PayLoadKey) const
 {
-	if (!SourceData)
+	check(PayloadSourceData == GetSourceData());
+
+	if (!GetSourceData())
 	{
 		UE_LOG(LogInterchangeImport, Error, TEXT("Failed to import TGA, bad source data."));
 		return {};
 	}
 
 	TArray64<uint8> SourceDataBuffer;
-	FString Filename = SourceData->GetFilename();
+	FString Filename = GetSourceData()->GetFilename();
 	
 	//Make sure the key fit the filename, The key should always be valid
 	if (!Filename.Equals(PayLoadKey))
diff --git a/Engine/Plugins/Runtime/Interchange/Source/Import/Private/Texture/InterchangeUDIMTranslator.cpp b/Engine/Plugins/Runtime/Interchange/Source/Import/Private/Texture/InterchangeUDIMTranslator.cpp
--- a/Engine/Plugins/Runtime/Interchange/Source/Import/Private/Texture/InterchangeUDIMTranslator.cpp
+++ b/Engine/Plugins/Runtime/Interchange/Source/Import/Private/Texture/InterchangeUDIMTranslator.cpp
@@ -160,9 +160,9 @@ bool UInterchangeUDIMTranslator::CanImportSourceData(const UInterchangeSourceDat
 	return false;
 }
 
-bool UInterchangeUDIMTranslator::Translate(const UInterchangeSourceData* SourceData, UInterchangeBaseNodeContainer& BaseNodeContainer) const
+bool UInterchangeUDIMTranslator::Translate(UInterchangeBaseNodeContainer& BaseNodeContainer) const
 {
-	FString Filename = SourceData->GetFilename();
+	FString Filename = GetSourceData()->GetFilename();
 	if (!FPaths::FileExists(Filename))
 	{
 		return false;
@@ -188,7 +188,7 @@ bool UInterchangeUDIMTranslator::Translate(const UInterchangeSourceData* SourceD
 	UInterchangeSourceData* TempSourceData = NewObject<UInterchangeSourceData>(GetTransientPackage(), UInterchangeSourceData::StaticClass());
 	UE::Interchange::Private::InterchangeUDIMTranslator::FScopedClearAsyncFlag AsyncFlagGuard(TempSourceData);
 
-	if (UInterchangeTranslatorBase* Translator = InterchangeManager.GetTranslatorSupportingPayloadInterfaceForSourceData(SourceData, UInterchangeTexturePayloadInterface::StaticClass()))
+	if (UInterchangeTranslatorBase* Translator = InterchangeManager.GetTranslatorSupportingPayloadInterfaceForSourceData(GetSourceData(), UInterchangeTexturePayloadInterface::StaticClass()))
 	{
 		UE::Interchange::Private::InterchangeUDIMTranslator::FScopedClearAsyncFlag AsyncFlagGuardTranslator(Translator);
 
diff --git a/Engine/Plugins/Runtime/Interchange/Source/Import/Public/Fbx/InterchangeFbxTranslator.h b/Engine/Plugins/Runtime/Interchange/Source/Import/Public/Fbx/InterchangeFbxTranslator.h
--- a/Engine/Plugins/Runtime/Interchange/Source/Import/Public/Fbx/InterchangeFbxTranslator.h
+++ b/Engine/Plugins/Runtime/Interchange/Source/Import/Public/Fbx/InterchangeFbxTranslator.h
@@ -29,7 +29,7 @@ public:
 
 	/** Begin UInterchangeTranslatorBase API*/
 	virtual bool CanImportSourceData(const UInterchangeSourceData* InSourceData) const override;
-	virtual bool Translate(const UInterchangeSourceData* SourceData, UInterchangeBaseNodeContainer& BaseNodeContainer) const override;
+	virtual bool Translate(UInterchangeBaseNodeContainer& BaseNodeContainer) const override;
 	virtual void ReleaseSource() override;
 	virtual void ImportFinish() override;
 	/** End UInterchangeTranslatorBase API*/
@@ -42,11 +42,11 @@ public:
 	 * Once the translation is done, the import process need a way to retrieve payload data.
 	 * This payload will be use by the factories to create the asset.
 	 *
-	 * @param SourceData - The source data containing the data to translate
+	 * @param InSourceData - The source data containing the data to translate
 	 * @param PayloadKey - The key to retrieve the a particular payload contain into the specified source data.
 	 * @return a PayloadData containing the imported data. The TOptional will not be set if there is an error.
 	 */
-	virtual TOptional<UE::Interchange::FImportImage> GetTexturePayloadData(const UInterchangeSourceData* SourceData, const FString& PayLoadKey) const override;
+	virtual TOptional<UE::Interchange::FImportImage> GetTexturePayloadData(const UInterchangeSourceData* InSourceData, const FString& PayLoadKey) const override;
 
 	/* IInterchangeTexturePayloadInterface End */
 
@@ -62,7 +62,7 @@ public:
 	 * @param PayloadKey - The key to retrieve the a particular payload contain into the specified source data.
 	 * @return a PayloadData containing the imported data. The TOptional will not be set if there is an error.
 	 */
-	virtual TOptional<UE::Interchange::FStaticMeshPayloadData> GetStaticMeshPayloadData(const FString& PayLoadKey) const override;
+	virtual TFuture<TOptional<UE::Interchange::FStaticMeshPayloadData>> GetStaticMeshPayloadData(const FString& PayLoadKey) const override;
 
 	/* IInterchangeStaticMeshPayloadInterface End */
 
diff --git a/Engine/Plugins/Runtime/Interchange/Source/Import/Public/Mesh/InterchangeStaticMeshFactory.h b/Engine/Plugins/Runtime/Interchange/Source/Import/Public/Mesh/InterchangeStaticMeshFactory.h
new file mode 100644
--- /dev/null
+++ b/Engine/Plugins/Runtime/Interchange/Source/Import/Public/Mesh/InterchangeStaticMeshFactory.h
@@ -0,0 +1,28 @@
+// Copyright Epic Games, Inc. All Rights Reserved.
+
+#pragma once
+
+#include "CoreMinimal.h"
+#include "InterchangeFactoryBase.h"
+#include "UObject/Object.h"
+#include "UObject/ObjectMacros.h"
+
+#include "InterchangeStaticMeshFactory.generated.h"
+
+UCLASS(BlueprintType)
+class INTERCHANGEIMPORT_API UInterchangeStaticMeshFactory : public UInterchangeFactoryBase
+{
+	GENERATED_BODY()
+public:
+
+	//////////////////////////////////////////////////////////////////////////
+	// Interchange factory base interface begin
+
+	virtual UClass* GetFactoryClass() const override;
+	virtual UObject* CreateEmptyAsset(const FCreateAssetParams& Arguments) const override;
+	virtual UObject* CreateAsset(const FCreateAssetParams& Arguments) const override;
+	virtual void PostImportGameThreadCallback(const FPostImportGameThreadCallbackParams& Arguments) const override;
+
+	// Interchange factory base interface end
+	//////////////////////////////////////////////////////////////////////////
+};
diff --git a/Engine/Plugins/Runtime/Interchange/Source/Import/Public/Mesh/InterchangeStaticMeshPayload.h b/Engine/Plugins/Runtime/Interchange/Source/Import/Public/Mesh/InterchangeStaticMeshPayload.h
--- a/Engine/Plugins/Runtime/Interchange/Source/Import/Public/Mesh/InterchangeStaticMeshPayload.h
+++ b/Engine/Plugins/Runtime/Interchange/Source/Import/Public/Mesh/InterchangeStaticMeshPayload.h
@@ -11,7 +11,8 @@ namespace UE
 	{
 		struct FStaticMeshPayloadData
 		{
-			TArray<FMeshDescription> MeshDescriptions;
+			FMeshDescription MeshDescription;
 		};
-	}//ns Interchange
-}//ns UE
+
+	} // namespace Interchange
+} // namespace UE
diff --git a/Engine/Plugins/Runtime/Interchange/Source/Import/Public/Mesh/InterchangeStaticMeshPayloadInterface.h b/Engine/Plugins/Runtime/Interchange/Source/Import/Public/Mesh/InterchangeStaticMeshPayloadInterface.h
--- a/Engine/Plugins/Runtime/Interchange/Source/Import/Public/Mesh/InterchangeStaticMeshPayloadInterface.h
+++ b/Engine/Plugins/Runtime/Interchange/Source/Import/Public/Mesh/InterchangeStaticMeshPayloadInterface.h
@@ -2,6 +2,7 @@
 #pragma once
 
 #include "CoreMinimal.h"
+#include "Async/Future.h"
 #include "InterchangeSourceData.h"
 #include "Mesh/InterchangeStaticMeshPayload.h"
 #include "UObject/Object.h"
@@ -31,7 +32,7 @@ public:
 	 * @param PayloadKey - The key to retrieve the a particular payload contain into the specified source data.
 	 * @return a PayloadData containing the the data ask with the key.
 	 */
-	virtual TOptional<UE::Interchange::FStaticMeshPayloadData> GetStaticMeshPayloadData(const FString& PayloadKey) const = 0;
+	virtual TFuture<TOptional<UE::Interchange::FStaticMeshPayloadData>> GetStaticMeshPayloadData(const FString& PayloadKey) const = 0;
 };
 
 
diff --git a/Engine/Plugins/Runtime/Interchange/Source/Import/Public/Texture/InterchangeBMPTranslator.h b/Engine/Plugins/Runtime/Interchange/Source/Import/Public/Texture/InterchangeBMPTranslator.h
--- a/Engine/Plugins/Runtime/Interchange/Source/Import/Public/Texture/InterchangeBMPTranslator.h
+++ b/Engine/Plugins/Runtime/Interchange/Source/Import/Public/Texture/InterchangeBMPTranslator.h
@@ -19,18 +19,17 @@ class INTERCHANGEIMPORT_API UInterchangeBMPTranslator : public UInterchangeTrans
 public:
 
 	/*
-	 * return true if the translator can translate the specified file.
+	 * return true if the translator can translate the specified source data.
 	 */
 	virtual bool CanImportSourceData(const UInterchangeSourceData* InSourceData) const override;
 
 	/**
-	 * Translate a source data into a node hold by the specified nodes container.
+	 * Translate the associated source data into a node hold by the specified nodes container.
 	 *
-	 * @param SourceData - The source data containing the data to translate
 	 * @param BaseNodeContainer - The unreal objects descriptions container where to put the translated source data.
 	 * @return true if the translator can translate the source data, false otherwise.
 	 */
-	virtual bool Translate(const UInterchangeSourceData* SourceData, UInterchangeBaseNodeContainer& BaseNodeContainer) const override;
+	virtual bool Translate(UInterchangeBaseNodeContainer& BaseNodeContainer) const override;
 
 
 	/* IInterchangeTexturePayloadInterface Begin */
@@ -39,11 +38,11 @@ public:
 	 * Once the translation is done, the import process need a way to retrieve payload data.
 	 * This payload will be use by the factories to create the asset.
 	 *
-	 * @param SourceData - The source data containing the data to translate
+	 * @param PayloadSourceData - The source data containing the data to translate
 	 * @param PayloadKey - The key to retrieve the a particular payload contain into the specified source data.
 	 * @return a PayloadData containing the import image data. The TOptional will not be set if there is an error.
 	 */
-	virtual TOptional<UE::Interchange::FImportImage> GetTexturePayloadData(const UInterchangeSourceData* SourceData, const FString& PayLoadKey) const override;
+	virtual TOptional<UE::Interchange::FImportImage> GetTexturePayloadData(const UInterchangeSourceData* PayloadSourceData, const FString& PayLoadKey) const override;
 
 	/* IInterchangeTexturePayloadInterface End */
 };
diff --git a/Engine/Plugins/Runtime/Interchange/Source/Import/Public/Texture/InterchangeDDSTranslator.h b/Engine/Plugins/Runtime/Interchange/Source/Import/Public/Texture/InterchangeDDSTranslator.h
--- a/Engine/Plugins/Runtime/Interchange/Source/Import/Public/Texture/InterchangeDDSTranslator.h
+++ b/Engine/Plugins/Runtime/Interchange/Source/Import/Public/Texture/InterchangeDDSTranslator.h
@@ -21,18 +21,17 @@ class INTERCHANGEIMPORT_API UInterchangeDDSTranslator : public UInterchangeTrans
 public:
 	
 	/*
-	 * return true if the translator can translate the specified file.
+	 * return true if the translator can translate the specified source data.
 	 */
 	virtual bool CanImportSourceData(const UInterchangeSourceData* InSourceData) const override;
 
 	/**
-	 * Translate a source data into a node hold by the specified nodes container.
+	 * Translate the associated source data into a node hold by the specified nodes container.
 	 *
-	 * @param SourceData - The source data containing the data to translate
 	 * @param BaseNodeContainer - The unreal objects descriptions container where to put the translated source data.
 	 * @return true if the translator can translate the source data, false otherwise.
 	 */
-	virtual bool Translate(const UInterchangeSourceData* SourceData, UInterchangeBaseNodeContainer& BaseNodeContainer) const override;
+	virtual bool Translate(UInterchangeBaseNodeContainer& BaseNodeContainer) const override;
 
 	/* IInterchangeTexturePayloadInterface Begin */
 
@@ -40,11 +39,11 @@ public:
 	 * Once the translation is done, the import process need a way to retrieve payload data.
 	 * This payload will be use by the factories to create the asset.
 	 *
-	 * @param SourceData - The source data containing the data to translate
+	 * @param PayloadSourceData - The source data containing the data to translate
 	 * @param PayloadKey - The key to retrieve the a particular payload contain into the specified source data.
 	 * @return a PayloadData containing the import image data. The TOptional will not be set if there is an error.
 	 */
-	virtual TOptional<UE::Interchange::FImportImage> GetTexturePayloadData(const UInterchangeSourceData* SourceData, const FString& PayLoadKey) const override;
+	virtual TOptional<UE::Interchange::FImportImage> GetTexturePayloadData(const UInterchangeSourceData* PayloadSourceData, const FString& PayLoadKey) const override;
 
 	/* IInterchangeTexturePayloadInterface End */
 
diff --git a/Engine/Plugins/Runtime/Interchange/Source/Import/Public/Texture/InterchangeHDRTranslator.h b/Engine/Plugins/Runtime/Interchange/Source/Import/Public/Texture/InterchangeHDRTranslator.h
--- a/Engine/Plugins/Runtime/Interchange/Source/Import/Public/Texture/InterchangeHDRTranslator.h
+++ b/Engine/Plugins/Runtime/Interchange/Source/Import/Public/Texture/InterchangeHDRTranslator.h
@@ -24,13 +24,12 @@ public:
 	virtual bool CanImportSourceData(const UInterchangeSourceData* InSourceData) const override;
 
 	/**
-	 * Translate a source data into a node hold by the specified nodes container.
+	 * Translate the associated source data into a node hold by the specified nodes container.
 	 *
-	 * @param SourceData - The source data containing the data to translate
 	 * @param BaseNodeContainer - The unreal objects descriptions container where to put the translated source data.
 	 * @return true if the translator can translate the source data, false otherwise.
 	 */
-	virtual bool Translate(const UInterchangeSourceData* SourceData, UInterchangeBaseNodeContainer& BaseNodeContainer) const override;
+	virtual bool Translate(UInterchangeBaseNodeContainer& BaseNodeContainer) const override;
 
 
 	/* IInterchangeTexturePayloadInterface Begin */
@@ -39,11 +38,11 @@ public:
 	 * Once the translation is done, the import process need a way to retrieve payload data.
 	 * This payload will be use by the factories to create the asset.
 	 *
-	 * @param SourceData - The source data containing the data to translate
+	 * @param PayloadSourceData - The source data containing the data to translate
 	 * @param PayloadKey - The key to retrieve the a particular payload contain into the specified source data.
 	 * @return a PayloadData containing the import image data. The TOptional will not be set if there is an error.
 	 */
-	virtual TOptional<UE::Interchange::FImportImage> GetTexturePayloadData(const UInterchangeSourceData* SourceData, const FString& PayLoadKey) const override;
+	virtual TOptional<UE::Interchange::FImportImage> GetTexturePayloadData(const UInterchangeSourceData* PayloadSourceData, const FString& PayLoadKey) const override;
 
 	/* IInterchangeTexturePayloadInterface End */
 };
diff --git a/Engine/Plugins/Runtime/Interchange/Source/Import/Public/Texture/InterchangeIESTranslator.h b/Engine/Plugins/Runtime/Interchange/Source/Import/Public/Texture/InterchangeIESTranslator.h
--- a/Engine/Plugins/Runtime/Interchange/Source/Import/Public/Texture/InterchangeIESTranslator.h
+++ b/Engine/Plugins/Runtime/Interchange/Source/Import/Public/Texture/InterchangeIESTranslator.h
@@ -24,13 +24,12 @@ public:
 	virtual bool CanImportSourceData(const UInterchangeSourceData* InSourceData) const override;
 
 	/**
-	 * Translate a source data into a node hold by the specified nodes container.
+	 * Translate the associated source data into a node hold by the specified nodes container.
 	 *
-	 * @param SourceData - The source data containing the data to translate
 	 * @param BaseNodeContainer - The unreal objects descriptions container where to put the translated source data.
 	 * @return true if the translator can translate the source data, false otherwise.
 	 */
-	virtual bool Translate(const UInterchangeSourceData* SourceData, UInterchangeBaseNodeContainer& BaseNodeContainer) const override;
+	virtual bool Translate(UInterchangeBaseNodeContainer& BaseNodeContainer) const override;
 
 
 	/* IInterchangeTexturePayloadInterface Begin */
@@ -43,7 +42,7 @@ public:
 	 * @param PayloadKey - The key to retrieve the a particular payload contain into the specified source data.
 	 * @return a PayloadData containing the import light profile data. The TOptional will not be set if there is an error.
 	 */
-	virtual TOptional<UE::Interchange::FImportLightProfile> GetLightProfilePayloadData(const UInterchangeSourceData* SourceData, const FString& PayLoadKey) const override;
+	virtual TOptional<UE::Interchange::FImportLightProfile> GetLightProfilePayloadData(const UInterchangeSourceData* PayloadSourceData, const FString& PayLoadKey) const override;
 
 	/* IInterchangeTexturePayloadInterface End */
 };
\ No newline at end of file
diff --git a/Engine/Plugins/Runtime/Interchange/Source/Import/Public/Texture/InterchangeJPGTranslator.h b/Engine/Plugins/Runtime/Interchange/Source/Import/Public/Texture/InterchangeJPGTranslator.h
--- a/Engine/Plugins/Runtime/Interchange/Source/Import/Public/Texture/InterchangeJPGTranslator.h
+++ b/Engine/Plugins/Runtime/Interchange/Source/Import/Public/Texture/InterchangeJPGTranslator.h
@@ -19,18 +19,17 @@ class INTERCHANGEIMPORT_API UInterchangeJPGTranslator : public UInterchangeTrans
 public:
 	
 	/**
-	 * return true if the translator can translate the specified file.
+	 * return true if the translator can translate the specified source data.
 	 */
 	virtual bool CanImportSourceData(const UInterchangeSourceData* InSourceData) const override;
 
 	/**
-	 * Translate a source data into a node hold by the specified nodes container.
+	 * Translate the associated source data into a node hold by the specified nodes container.
 	 *
-	 * @param InSourceData - The source data containing the JPG data
 	 * @param BaseNodeContainer - The unreal objects descriptions container where to put the translated source data.
 	 * @return true if the translator can translate the source data, false otherwise.
 	 */
-	virtual bool Translate(const UInterchangeSourceData* SourceData, UInterchangeBaseNodeContainer& BaseNodeContainer) const override;
+	virtual bool Translate(UInterchangeBaseNodeContainer& BaseNodeContainer) const override;
 
 	/* IInterchangeTexturePayloadInterface Begin */
 
@@ -38,11 +37,11 @@ public:
 	 * Once the translation is done, the import process need a way to retrieve payload data.
 	 * This payload will be use by the factories to create the asset.
 	 *
-	 * @param SourceData - The source data containing the data to translate
+	 * @param PayloadSourceData - The source data containing the data to translate
 	 * @param PayloadKey - The key to retrieve the a particular payload contain into the specified source data.
 	 * @return a PayloadData containing the import image data. The TOptional will not be set if there is an error.
 	 */
-	virtual TOptional<UE::Interchange::FImportImage> GetTexturePayloadData(const UInterchangeSourceData* SourceData, const FString& PayLoadKey) const override;
+	virtual TOptional<UE::Interchange::FImportImage> GetTexturePayloadData(const UInterchangeSourceData* PayloadSourceData, const FString& PayLoadKey) const override;
 
 	/* IInterchangeTexturePayloadInterface End */
 };
diff --git a/Engine/Plugins/Runtime/Interchange/Source/Import/Public/Texture/InterchangePCXTranslator.h b/Engine/Plugins/Runtime/Interchange/Source/Import/Public/Texture/InterchangePCXTranslator.h
--- a/Engine/Plugins/Runtime/Interchange/Source/Import/Public/Texture/InterchangePCXTranslator.h
+++ b/Engine/Plugins/Runtime/Interchange/Source/Import/Public/Texture/InterchangePCXTranslator.h
@@ -19,18 +19,17 @@ class INTERCHANGEIMPORT_API UInterchangePCXTranslator : public UInterchangeTrans
 public:
 	
 	/*
-	 * return true if the translator can translate the specified file.
+	 * return true if the translator can translate the specified source data.
 	 */
 	virtual bool CanImportSourceData(const UInterchangeSourceData* InSourceData) const override;
 
 	/**
-	 * Translate a source data into a node hold by the specified nodes container.
+	 * Translates the associated source data into a node hold by the specified nodes container.
 	 *
-	 * @param SourceData - The source data containing the data to translate
 	 * @param BaseNodeContainer - The unreal objects descriptions container where to put the translated source data.
 	 * @return true if the translator can translate the source data, false otherwise.
 	 */
-	virtual bool Translate(const UInterchangeSourceData* SourceData, UInterchangeBaseNodeContainer& BaseNodeContainer) const override;
+	virtual bool Translate(UInterchangeBaseNodeContainer& BaseNodeContainer) const override;
 
 	/* IInterchangeTexturePayloadInterface Begin */
 
@@ -38,11 +37,11 @@ public:
 	 * Once the translation is done, the import process need a way to retrieve payload data.
 	 * This payload will be use by the factories to create the asset.
 	 *
-	 * @param SourceData - The source data containing the data to translate
+	 * @param PayloadSourceData - The source data containing the data to translate
 	 * @param PayloadKey - The key to retrieve the a particular payload contain into the specified source data.
 	 * @return a PayloadData containing the import image data. The TOptional will not be set if there is an error.
 	 */
-	virtual TOptional<UE::Interchange::FImportImage> GetTexturePayloadData(const UInterchangeSourceData* SourceData, const FString& PayLoadKey) const override;
+	virtual TOptional<UE::Interchange::FImportImage> GetTexturePayloadData(const UInterchangeSourceData* PayloadSourceData, const FString& PayLoadKey) const override;
 
 	/* IInterchangeTexturePayloadInterface End */
 };
diff --git a/Engine/Plugins/Runtime/Interchange/Source/Import/Public/Texture/InterchangeTGATranslator.h b/Engine/Plugins/Runtime/Interchange/Source/Import/Public/Texture/InterchangeTGATranslator.h
--- a/Engine/Plugins/Runtime/Interchange/Source/Import/Public/Texture/InterchangeTGATranslator.h
+++ b/Engine/Plugins/Runtime/Interchange/Source/Import/Public/Texture/InterchangeTGATranslator.h
@@ -18,18 +18,17 @@ class INTERCHANGEIMPORT_API UInterchangeTGATranslator : public UInterchangeTrans
 public:
 	
 	/**
-	 * return true if the translator can translate the specified file.
+	 * return true if the translator can translate the specified source data.
 	 */
 	virtual bool CanImportSourceData(const UInterchangeSourceData* InSourceData) const override;
 
 	/**
-	 * Translate a source data into a node hold by the specified nodes container.
+	 * Translates the associated source data into a node hold by the specified nodes container.
 	 *
-	 * @param InSourceData - The source data containing the JPG data
 	 * @param BaseNodeContainer - The unreal objects descriptions container where to put the translated source data.
 	 * @return true if the translator can translate the source data, false otherwise.
 	 */
-	virtual bool Translate(const UInterchangeSourceData* SourceData, UInterchangeBaseNodeContainer& BaseNodeContainer) const override;
+	virtual bool Translate(UInterchangeBaseNodeContainer& BaseNodeContainer) const override;
 
 	/* IInterchangeTexturePayloadInterface Begin */
 
@@ -37,11 +36,11 @@ public:
 	 * Once the translation is done, the import process need a way to retrieve payload data.
 	 * This payload will be use by the factories to create the asset.
 	 *
-	 * @param SourceData - The source data containing the data to translate
+	 * @param PayloadSourceData - The source data containing the data to translate
 	 * @param PayloadKey - The key to retrieve the a particular payload contain into the specified source data.
 	 * @return a PayloadData containing the import image data. The TOptional will not be set if there is an error.
 	 */
-	virtual TOptional<UE::Interchange::FImportImage> GetTexturePayloadData(const UInterchangeSourceData* SourceData, const FString& PayLoadKey) const override;
+	virtual TOptional<UE::Interchange::FImportImage> GetTexturePayloadData(const UInterchangeSourceData* PayloadSourceData, const FString& PayLoadKey) const override;
 
 	/* IInterchangeTexturePayloadInterface End */
 };
diff --git a/Engine/Plugins/Runtime/Interchange/Source/Import/Public/Texture/InterchangeTexturePayloadInterface.h b/Engine/Plugins/Runtime/Interchange/Source/Import/Public/Texture/InterchangeTexturePayloadInterface.h
--- a/Engine/Plugins/Runtime/Interchange/Source/Import/Public/Texture/InterchangeTexturePayloadInterface.h
+++ b/Engine/Plugins/Runtime/Interchange/Source/Import/Public/Texture/InterchangeTexturePayloadInterface.h
@@ -27,11 +27,11 @@ public:
 	 * Once the translation is done, the import process need a way to retrieve payload data.
 	 * This payload will be use by the factories to create the asset.
 	 *
-	 * @param SourceData - The source data containing the data to translate
+	 * @param PayloadSourceData - The source data containing the data to translate
 	 * @param PayloadKey - The key to retrieve the a particular payload contain into the specified source data.
 	 * @return a PayloadData containing the import image data. The TOptional will not be set if there is an error.
 	 */
-	virtual TOptional<UE::Interchange::FImportImage> GetTexturePayloadData(const UInterchangeSourceData* SourceData, const FString& PayloadKey) const = 0;
+	virtual TOptional<UE::Interchange::FImportImage> GetTexturePayloadData(const UInterchangeSourceData* PayloadSourceData, const FString& PayloadKey) const = 0;
 };
 
 
