First pass implementaiton of a message passing system for Interchange. This permits messages, warnings and errors to be output from translators, pipelines and factories, as well as from external proce

diff --git a/Engine/Plugins/Runtime/Interchange/Source/Import/Public/Texture/InterchangeUDIMTranslator.h b/Engine/Plugins/Runtime/Interchange/Source/Import/Public/Texture/InterchangeUDIMTranslator.h
--- a/Engine/Plugins/Runtime/Interchange/Source/Import/Public/Texture/InterchangeUDIMTranslator.h
+++ b/Engine/Plugins/Runtime/Interchange/Source/Import/Public/Texture/InterchangeUDIMTranslator.h
@@ -25,18 +25,17 @@ public:
 	UInterchangeUDIMTranslator();
 
 	/*
-	 * return true if the translator can translate the specified file.
+	 * return true if the translator can translate the specified source data.
 	 */
 	virtual bool CanImportSourceData(const UInterchangeSourceData* InSourceData) const override;
 
 	/**
-	 * Translate a source data into a node hold by the specified nodes container.
+	 * Translates the associated source data into a node hold by the specified nodes container.
 	 *
-	 * @param SourceData - The source data containing the data to translate
 	 * @param BaseNodeContainer - The unreal objects descriptions container where to put the translated source data.
 	 * @return true if the translator can translate the source data, false otherwise.
 	 */
-	virtual bool Translate(const UInterchangeSourceData* SourceData, UInterchangeBaseNodeContainer& BaseNodeContainer) const override;
+	virtual bool Translate(UInterchangeBaseNodeContainer& BaseNodeContainer) const override;
 
 	/**
 	 * Once the translation is done, the import process need a way to retrieve payload data.
diff --git a/Engine/Plugins/Runtime/Interchange/Source/Messages/Public/Fbx/InterchangeFbxMessages.h b/Engine/Plugins/Runtime/Interchange/Source/Messages/Public/Fbx/InterchangeFbxMessages.h
new file mode 100644
--- /dev/null
+++ b/Engine/Plugins/Runtime/Interchange/Source/Messages/Public/Fbx/InterchangeFbxMessages.h
@@ -0,0 +1,87 @@
+// Copyright Epic Games, Inc. All Rights Reserved.
+
+#pragma once
+
+#include "CoreMinimal.h"
+#include "InterchangeResult.h"
+
+#include "InterchangeFbxMessages.generated.h"
+
+
+/**
+ * Base class for FBX parser warnings
+ */
+UCLASS()
+class INTERCHANGEMESSAGES_API UInterchangeResultMeshWarning : public UInterchangeResultWarning
+{
+	GENERATED_BODY()
+
+public:
+
+	UPROPERTY()
+	FString MeshName;
+};
+
+
+/**
+ * Base class for FBX parser errors
+ */
+UCLASS()
+class INTERCHANGEMESSAGES_API UInterchangeResultMeshError : public UInterchangeResultError
+{
+	GENERATED_BODY()
+
+public:
+
+	UPROPERTY()
+	FString MeshName;
+};
+
+
+/**
+ * A generic class for FBX parser warnings, with no additional metadata, and where the text is specified by the user
+ */
+UCLASS()
+class INTERCHANGEMESSAGES_API UInterchangeResultMeshWarning_Generic : public UInterchangeResultMeshWarning
+{
+	GENERATED_BODY()
+
+public:
+	virtual FText GetText() const override;
+
+	UPROPERTY()
+	FText Text;
+};
+
+
+/**
+ * A generic class for FBX parser errors, with no additional metadata, and where the text is specified by the user
+ */
+UCLASS()
+class INTERCHANGEMESSAGES_API UInterchangeResultMeshError_Generic : public UInterchangeResultMeshError
+{
+	GENERATED_BODY()
+
+public:
+	virtual FText GetText() const override;
+
+	UPROPERTY()
+	FText Text;
+};
+
+
+/**
+ * 
+ */
+UCLASS()
+class INTERCHANGEMESSAGES_API UInterchangeResultMeshWarning_TooManyUVs : public UInterchangeResultMeshWarning
+{
+	GENERATED_BODY()
+
+public:
+	virtual FText GetText() const override;
+
+	UPROPERTY()
+	int32 ExcessUVs;
+};
+
diff --git a/Engine/Plugins/Runtime/Interchange/Source/Nodes/Public/InterchangeStaticMeshFactoryNode.h b/Engine/Plugins/Runtime/Interchange/Source/Nodes/Public/InterchangeStaticMeshFactoryNode.h
new file mode 100644
--- /dev/null
+++ b/Engine/Plugins/Runtime/Interchange/Source/Nodes/Public/InterchangeStaticMeshFactoryNode.h
@@ -0,0 +1,191 @@
+// Copyright Epic Games, Inc. All Rights Reserved.
+
+#pragma once
+
+#include "CoreMinimal.h"
+#include "Nodes/InterchangeBaseNode.h"
+
+#if WITH_ENGINE
+#include "Engine/StaticMesh.h"
+#endif
+
+#include "InterchangeStaticMeshFactoryNode.generated.h"
+
+
+namespace UE
+{
+	namespace Interchange
+	{
+		struct FStaticMeshNodeStaticData : public FBaseNodeStaticData
+		{
+			static const FString& GetLodDependenciesBaseKey()
+			{
+				static FString LodDependencies_BaseKey = TEXT("Lod_Dependencies");
+				return LodDependencies_BaseKey;
+			}
+		};
+
+	} // namespace Interchange
+} // namespace UE
+
+
+UCLASS(BlueprintType)
+class INTERCHANGENODES_API UInterchangeStaticMeshFactoryNode : public UInterchangeBaseNode
+{
+	GENERATED_BODY()
+
+public:
+	UInterchangeStaticMeshFactoryNode()
+	{
+#if WITH_ENGINE
+		AssetClass = nullptr;
+#endif
+		LodDependencies.Initialize(Attributes, UE::Interchange::FStaticMeshNodeStaticData::GetLodDependenciesBaseKey());
+	}
+
+	/**
+	 * Initialize node data
+	 * @param UniqueID - The uniqueId for this node
+	 * @param DisplayLabel - The name of the node
+	 * @param InAssetClass - The class the StaticMesh factory will create for this node.
+	 */
+	UFUNCTION(BlueprintCallable, Category = "Interchange | Node | StaticMesh")
+	void InitializeStaticMeshNode(const FString& UniqueID, const FString& DisplayLabel, const FString& InAssetClass)
+	{
+		bIsNodeClassInitialized = false;
+		InitializeNode(UniqueID, DisplayLabel, EInterchangeNodeContainerType::NodeContainerType_FactoryData);
+
+		FString OperationName = GetTypeName() + TEXT(".SetAssetClassName");
+		InterchangePrivateNodeBase::SetCustomAttribute<FString>(*Attributes, ClassNameAttributeKey, OperationName, InAssetClass);
+		FillAssetClassFromAttribute();
+	}
+
+	virtual void Serialize(FArchive& Ar) override
+	{
+		Super::Serialize(Ar);
+#if WITH_ENGINE
+		if (Ar.IsLoading())
+		{
+			// Make sure the class is properly set when we compile with engine, this will set the bIsNodeClassInitialized to true.
+			SetNodeClassFromClassAttribute();
+		}
+#endif //#if WITH_ENGINE
+	}
+
+	/**
+	 * Return the node type name of the class, we use this when reporting error
+	 */
+	virtual FString GetTypeName() const override
+	{
+		const FString TypeName = TEXT("StaticMeshNode");
+		return TypeName;
+	}
+
+	virtual FString GetKeyDisplayName(const UE::Interchange::FAttributeKey& NodeAttributeKey) const override
+	{
+		FString KeyDisplayName = NodeAttributeKey.Key;
+		if (NodeAttributeKey.Key.Equals(UE::Interchange::FStaticMeshNodeStaticData::GetLodDependenciesBaseKey()))
+		{
+			KeyDisplayName = TEXT("LOD Dependencies Count");
+			return KeyDisplayName;
+		}
+		else if (NodeAttributeKey.Key.StartsWith(UE::Interchange::FStaticMeshNodeStaticData::GetLodDependenciesBaseKey()))
+		{
+			KeyDisplayName = TEXT("LOD Dependencies Index ");
+			const FString IndexKey = UE::Interchange::FNameAttributeArrayHelper::IndexKey();
+			int32 IndexPosition = NodeAttributeKey.Key.Find(IndexKey) + IndexKey.Len();
+			if (IndexPosition < NodeAttributeKey.Key.Len())
+			{
+				KeyDisplayName += NodeAttributeKey.Key.RightChop(IndexPosition);
+			}
+			return KeyDisplayName;
+		}
+		return Super::GetKeyDisplayName(NodeAttributeKey);
+	}
+
+	/** Get the class this node want to create */
+	UFUNCTION(BlueprintCallable, Category = "Interchange | Node | StaticMesh")
+	virtual class UClass* GetAssetClass() const override
+	{
+		ensure(bIsNodeClassInitialized);
+#if WITH_ENGINE
+		return AssetClass.Get() != nullptr ? AssetClass.Get() : UStaticMesh::StaticClass();
+#else
+		return nullptr;
+#endif
+	}
+
+	virtual FGuid GetHash() const override
+	{
+		return Attributes->GetStorageHash();
+	}
+
+public:
+	/** Return The number of LOD this static mesh has.*/
+	UFUNCTION(BlueprintCallable, Category = "Interchange | Node | StaticMesh")
+	int32 GetLodDataCount() const
+	{
+		return LodDependencies.GetCount();
+	}
+
+	UFUNCTION(BlueprintCallable, Category = "Interchange | Node | StaticMesh")
+	void GetLodDataUniqueIds(TArray<FString>& OutLodDataUniqueIds) const
+	{
+		LodDependencies.GetNames(OutLodDataUniqueIds);
+	}
+
+	UFUNCTION(BlueprintCallable, Category = "Interchange | Node | StaticMesh")
+	bool AddLodDataUniqueId(const FString& LodDataUniqueId)
+	{
+		return LodDependencies.AddName(LodDataUniqueId);
+	}
+
+	UFUNCTION(BlueprintCallable, Category = "Interchange | Node | StaticMesh")
+	bool RemoveLodDataUniqueId(const FString& LodDataUniqueId)
+	{
+		return LodDependencies.RemoveName(LodDataUniqueId);
+	}
+
+private:
+
+	void FillAssetClassFromAttribute()
+	{
+#if WITH_ENGINE
+		FString OperationName = GetTypeName() + TEXT(".GetAssetClassName");
+		FString ClassName;
+		InterchangePrivateNodeBase::GetCustomAttribute<FString>(*Attributes, ClassNameAttributeKey, OperationName, ClassName);
+		if (ClassName.Equals(UStaticMesh::StaticClass()->GetName()))
+		{
+			AssetClass = UStaticMesh::StaticClass();
+			bIsNodeClassInitialized = true;
+		}
+#endif
+	}
+
+	bool SetNodeClassFromClassAttribute()
+	{
+		if (!bIsNodeClassInitialized)
+		{
+			FillAssetClassFromAttribute();
+		}
+		return bIsNodeClassInitialized;
+	}
+
+	bool IsEditorOnlyDataDefined()
+	{
+#if WITH_EDITORONLY_DATA
+		return true;
+#else
+		return false;
+#endif
+	}
+
+	const UE::Interchange::FAttributeKey ClassNameAttributeKey = UE::Interchange::FBaseNodeStaticData::ClassTypeAttributeKey();
+
+	UE::Interchange::FNameAttributeArrayHelper LodDependencies;
+protected:
+#if WITH_ENGINE
+	TSubclassOf<UStaticMesh> AssetClass = nullptr;
+#endif
+	bool bIsNodeClassInitialized = false;
+};
diff --git a/Engine/Plugins/Runtime/Interchange/Source/Nodes/Public/InterchangeStaticMeshLodDataNode.h b/Engine/Plugins/Runtime/Interchange/Source/Nodes/Public/InterchangeStaticMeshLodDataNode.h
new file mode 100644
--- /dev/null
+++ b/Engine/Plugins/Runtime/Interchange/Source/Nodes/Public/InterchangeStaticMeshLodDataNode.h
@@ -0,0 +1,132 @@
+// Copyright Epic Games, Inc. All Rights Reserved.
+
+#pragma once
+
+#include "CoreMinimal.h"
+#include "Nodes/InterchangeBaseNode.h"
+
+#include "InterchangeStaticMeshLodDataNode.generated.h"
+
+
+namespace UE
+{
+	namespace Interchange
+	{
+		struct FStaticMeshNodeLodDataStaticData : public FBaseNodeStaticData
+		{
+			static const FString& GetMeshUidsBaseKey()
+			{
+				static FString MeshUids_BaseKey = TEXT("__MeshUids__Key");
+				return MeshUids_BaseKey;
+			}
+		};
+
+	} // namespace Interchange
+} // namespace UE
+
+
+UCLASS(BlueprintType)
+class INTERCHANGENODES_API UInterchangeStaticMeshLodDataNode : public UInterchangeBaseNode
+{
+	GENERATED_BODY()
+
+public:
+	UInterchangeStaticMeshLodDataNode()
+	{
+		MeshUids.Initialize(Attributes, UE::Interchange::FStaticMeshNodeLodDataStaticData::GetMeshUidsBaseKey());
+	}
+
+	/**
+	 * Return the node type name of the class, we use this when reporting error
+	 */
+	virtual FString GetTypeName() const override
+	{
+		const FString TypeName = TEXT("StaticMeshLodDataNode");
+		return TypeName;
+	}
+
+	virtual FString GetKeyDisplayName(const UE::Interchange::FAttributeKey& NodeAttributeKey) const override
+	{
+		FString KeyDisplayName = NodeAttributeKey.Key;
+		if (NodeAttributeKey.Key.Equals(UE::Interchange::FStaticMeshNodeLodDataStaticData::GetMeshUidsBaseKey()))
+		{
+			KeyDisplayName = TEXT("Mesh count");
+			return KeyDisplayName;
+		}
+		else if (NodeAttributeKey.Key.StartsWith(UE::Interchange::FStaticMeshNodeLodDataStaticData::GetMeshUidsBaseKey()))
+		{
+			KeyDisplayName = TEXT("Mesh index ");
+			const FString IndexKey = UE::Interchange::FNameAttributeArrayHelper::IndexKey();
+			int32 IndexPosition = NodeAttributeKey.Key.Find(IndexKey) + IndexKey.Len();
+			if (IndexPosition < NodeAttributeKey.Key.Len())
+			{
+				KeyDisplayName += NodeAttributeKey.Key.RightChop(IndexPosition);
+			}
+			return KeyDisplayName;
+		}
+
+		return Super::GetKeyDisplayName(NodeAttributeKey);
+	}
+
+	virtual FString GetAttributeCategory(const UE::Interchange::FAttributeKey& NodeAttributeKey) const override
+	{
+		if (NodeAttributeKey.Key.StartsWith(UE::Interchange::FStaticMeshNodeLodDataStaticData::GetMeshUidsBaseKey()))
+		{
+			return FString(TEXT("Meshes"));
+		}
+		return Super::GetAttributeCategory(NodeAttributeKey);
+	}
+
+	virtual FGuid GetHash() const override
+	{
+		return Attributes->GetStorageHash();
+	}
+
+
+public:
+	/* Mesh Uids: It can be either a scene or a mesh node uid. If its a scene it mean we want the mesh factory to bake the geo payload with the global transform of the scene node. */
+
+	UFUNCTION(BlueprintCallable, Category = "Interchange | Node | StaticMeshLodData")
+	int32 GetMeshUidsCount() const
+	{
+		return MeshUids.GetCount();
+	}
+
+	UFUNCTION(BlueprintCallable, Category = "Interchange | Node | StaticMeshLodData")
+	void GetMeshUids(TArray<FString>& OutMeshNames) const
+	{
+		MeshUids.GetNames(OutMeshNames);
+	}
+
+	UFUNCTION(BlueprintCallable, Category = "Interchange | Node | StaticMeshLodData")
+	bool AddMeshUid(const FString& MeshName)
+	{
+		return MeshUids.AddName(MeshName);
+	}
+
+	UFUNCTION(BlueprintCallable, Category = "Interchange | Node | StaticMeshLodData")
+	bool RemoveMeshUid(const FString& MeshName)
+	{
+		return MeshUids.RemoveName(MeshName);
+	}
+
+	UFUNCTION(BlueprintCallable, Category = "Interchange | Node | StaticMeshLodData")
+	bool RemoveAllMeshes()
+	{
+		return MeshUids.RemoveAllNames();
+	}
+
+private:
+
+	bool IsEditorOnlyDataDefined()
+	{
+#if WITH_EDITORONLY_DATA
+		return true;
+#else
+		return false;
+#endif
+	}
+
+	UE::Interchange::FNameAttributeArrayHelper MeshUids;
+protected:
+};
diff --git a/Engine/Plugins/Runtime/Interchange/Source/Parsers/Fbx/InterchangeFbxParser.Build.cs b/Engine/Plugins/Runtime/Interchange/Source/Parsers/Fbx/InterchangeFbxParser.Build.cs
--- a/Engine/Plugins/Runtime/Interchange/Source/Parsers/Fbx/InterchangeFbxParser.Build.cs
+++ b/Engine/Plugins/Runtime/Interchange/Source/Parsers/Fbx/InterchangeFbxParser.Build.cs
@@ -22,6 +22,7 @@ namespace UnrealBuildTool.Rules
 			PrivateDependencyModuleNames.AddRange(
 				new string[]
 				{
+					"InterchangeMessages",
 					"InterchangeNodes",
 					"MeshDescription",
 					"SkeletalMeshDescription",
diff --git a/Engine/Plugins/Runtime/Interchange/Source/Parsers/Fbx/Private/FbxAPI.cpp b/Engine/Plugins/Runtime/Interchange/Source/Parsers/Fbx/Private/FbxAPI.cpp
--- a/Engine/Plugins/Runtime/Interchange/Source/Parsers/Fbx/Private/FbxAPI.cpp
+++ b/Engine/Plugins/Runtime/Interchange/Source/Parsers/Fbx/Private/FbxAPI.cpp
@@ -13,6 +13,8 @@
 #include "InterchangeTextureNode.h"
 #include "Nodes/InterchangeBaseNodeContainer.h"
 
+#define LOCTEXT_NAMESPACE "InterchangeFbxParser"
+
 namespace UE
 {
 	namespace Interchange
@@ -24,7 +26,7 @@ namespace UE
 				PayloadContexts.Empty();
 			}
 
-			bool FFbxParser::LoadFbxFile(const FString& Filename, TArray<FString>& JSonErrorMessages)
+			bool FFbxParser::LoadFbxFile(const FString& Filename)
 			{
 				SourceFilename = Filename;
 				int32 SDKMajor, SDKMinor, SDKRevision;
@@ -33,7 +35,8 @@ namespace UE
 				SDKManager = FbxManager::Create();
 				if (!SDKManager)
 				{
-					JSonErrorMessages.Add(TEXT("{\"Msg\" : {\"Type\" : \"Error\",\n\"Msg\" : \"Cannot create fbx sdk manager\"}}"));
+					UInterchangeResultError_Generic* Message = AddMessage<UInterchangeResultError_Generic>();
+					Message->Text = LOCTEXT("CannotCreateFBXManager", "Cannot create FBX SDK manager.");
 					return false;
 				}
 
@@ -57,7 +60,8 @@ namespace UE
 				SDKScene = FbxScene::Create(SDKManager, "My Scene");
 				if (!SDKScene)
 				{
-					JSonErrorMessages.Add(TEXT("{\"Msg\" : {\"Type\" : \"Error\",\n\"Msg\" : \"Cannot create fbx sdk scene\"}}"));
+					UInterchangeResultError_Generic* Message = AddMessage<UInterchangeResultError_Generic>();
+					Message->Text = LOCTEXT("CannotCreateFBXScene", "Cannot create FBX SDK scene.");
 					return false;
 				}
 				
@@ -72,7 +76,12 @@ namespace UE
 				const bool bImportStatus = SDKImporter->Initialize(TCHAR_TO_UTF8(*Filename));
 				if (!bImportStatus)
 				{
-					JSonErrorMessages.Add(TEXT("{\"Msg\" : {\"Type\" : \"Error\",\n\"Msg\" : \"Cannot open fbx filename\"}}"));
+					FFormatNamedArguments FilenameText
+					{
+						{ TEXT("Filename"), FText::FromString(Filename) }
+					};
+					UInterchangeResultError_Generic* Message = AddMessage<UInterchangeResultError_Generic>();
+					Message->Text = FText::Format(LOCTEXT("CannotOpenFBXFile", "Cannot open FBX file '{Filename}'."), FilenameText);
 					return false;
 				}
 
@@ -89,28 +98,35 @@ namespace UE
 				return true;
 			}
 
-			void FFbxParser::FillContainerWithFbxScene(UInterchangeBaseNodeContainer& NodeContainer, TArray<FString>& JSonErrorMessages)
+			void FFbxParser::FillContainerWithFbxScene(UInterchangeBaseNodeContainer& NodeContainer)
 			{
-				CleanupFbxData(JSonErrorMessages);
-				FFbxMaterial::AddAllTextures(SDKScene, NodeContainer, JSonErrorMessages);
-				FFbxMaterial::AddAllMaterials(SDKScene, NodeContainer, JSonErrorMessages);
-				FFbxMesh::AddAllMeshes(SDKScene, SDKGeometryConverter, NodeContainer, JSonErrorMessages, PayloadContexts);
-				FFbxScene::AddHierarchy(SDKScene, NodeContainer, JSonErrorMessages);
+				CleanupFbxData();
+
+				FFbxMaterial FbxMaterial(*this);
+				FbxMaterial.AddAllTextures(SDKScene, NodeContainer);
+				FbxMaterial.AddAllMaterials(SDKScene, NodeContainer);
+
+				FFbxMesh FbxMesh(*this);
+				FbxMesh.AddAllMeshes(SDKScene, SDKGeometryConverter, NodeContainer, PayloadContexts);
+
+				FFbxScene FbxScene(*this);
+				FbxScene.AddHierarchy(SDKScene, NodeContainer);
 			}
 
-			bool FFbxParser::FetchPayloadData(const FString& PayloadKey, const FString& PayloadFilepath, TArray<FString>& JSonErrorMessages)
+			bool FFbxParser::FetchPayloadData(const FString& PayloadKey, const FString& PayloadFilepath)
 			{
 				if (!PayloadContexts.Contains(PayloadKey))
 				{
-					JSonErrorMessages.Add(TEXT("{\"Msg\" : {\"Type\" : \"Error\",\n\"Msg\" : \"Cannot retrieve payload, payload key doesnt have any context.\"}}"));
+					UInterchangeResultError_Generic* Message = AddMessage<UInterchangeResultError_Generic>();
+					Message->Text = LOCTEXT("CannotRetrievePayload", "Cannot retrieve payload; payload key doesn't have any context.");
 					return false;
 				}
 
 				TSharedPtr<FPayloadContextBase>& PayloadContext = PayloadContexts.FindChecked(PayloadKey);
-				return PayloadContext->FetchPayloadToFile(PayloadFilepath, JSonErrorMessages);
+				return PayloadContext->FetchPayloadToFile(*this, PayloadFilepath);
 			}
 
-			void FFbxParser::CleanupFbxData(TArray<FString>& JSonErrorMessages)
+			void FFbxParser::CleanupFbxData()
 			{
 				//////////////////////////////////////////////////////////////////////////
 				// Make sure there is a valid bind pose
@@ -123,7 +139,8 @@ namespace UE
 				const int32 NbPoses = SDKScene->GetFbxManager()->GetBindPoseCount(SDKScene);
 				if (NbPoses != Default_NbPoses)
 				{
-					JSonErrorMessages.Add(TEXT("{\"Msg\" : {\"Type\" : \"Warning\",\n\"Msg\" : \"FbxParser: Missing bind pose, the fbx sdk create one.\"}}"));
+					UInterchangeResultWarning_Generic* Message = AddMessage<UInterchangeResultWarning_Generic>();
+					Message->Text = LOCTEXT("MissingBindPose", "Missing bind pose - the FBX SDK has created one.");
 				}
 
 				//////////////////////////////////////////////////////////////////////////
@@ -132,3 +149,5 @@ namespace UE
 		} //ns Private
 	} //ns Interchange
 } //ns UE
+
+#undef LOCTEXT_NAMESPACE
diff --git a/Engine/Plugins/Runtime/Interchange/Source/Parsers/Fbx/Private/FbxAPI.h b/Engine/Plugins/Runtime/Interchange/Source/Parsers/Fbx/Private/FbxAPI.h
--- a/Engine/Plugins/Runtime/Interchange/Source/Parsers/Fbx/Private/FbxAPI.h
+++ b/Engine/Plugins/Runtime/Interchange/Source/Parsers/Fbx/Private/FbxAPI.h
@@ -4,6 +4,8 @@
 
 #include "CoreMinimal.h"
 #include "FbxInclude.h"
+#include "InterchangeResultsContainer.h"
+
 
 #define FBX_METADATA_PREFIX TEXT("FBX.")
 #define INVALID_UNIQUE_ID 0xFFFFFFFFFFFFFFFF
@@ -29,22 +31,45 @@ namespace UE
 			class FFbxParser
 			{
 			public:
+				explicit FFbxParser(UInterchangeResultsContainer* InResultsContainer)
+					: ResultsContainer(InResultsContainer)
+				{}
+
 				~FFbxParser();
 
 				/* Load an fbx file into the fbx sdk, return false if the file could not be load. */
-				bool LoadFbxFile(const FString& Filename, TArray<FString>& JsonErrorMessages);
+				bool LoadFbxFile(const FString& Filename);
 
 				/* Extract the fbx data from the sdk into our node container */
-				void FillContainerWithFbxScene(UInterchangeBaseNodeContainer& NodeContainer, TArray<FString>& JsonErrorMessages);
+				void FillContainerWithFbxScene(UInterchangeBaseNodeContainer& NodeContainer);
 
 				/* Extract the fbx data from the sdk into our node container */
-				bool FetchPayloadData(const FString& PayloadKey, const FString& PayloadFilepath, TArray<FString>& JSonErrorMessages);
+				bool FetchPayloadData(const FString& PayloadKey, const FString& PayloadFilepath);
 			
-			private:
+				/**
+				 * This function is used to add the given message object directly into the results for this operation.
+				 */
+				template <typename T>
+				T* AddMessage() const
+				{
+					check(ResultsContainer != nullptr);
+					T* Item = ResultsContainer->Add<T>();
+					Item->SourceAssetName = SourceFilename;
+					return Item;
+				}
 
-				void CleanupFbxData(TArray<FString>& JSonErrorMessages);
 
+				void AddMessage(UInterchangeResult* Item) const
+				{
+					check(ResultsContainer != nullptr);
+					ResultsContainer->Add(Item);
+					Item->SourceAssetName = SourceFilename;
+				}
+			private:
+
+				void CleanupFbxData();
 
+				UInterchangeResultsContainer* ResultsContainer;
 				FbxManager* SDKManager = nullptr;
 				FbxScene* SDKScene = nullptr;
 				FbxImporter* SDKImporter = nullptr;
diff --git a/Engine/Plugins/Runtime/Interchange/Source/Parsers/Fbx/Private/FbxHelper.cpp b/Engine/Plugins/Runtime/Interchange/Source/Parsers/Fbx/Private/FbxHelper.cpp
--- a/Engine/Plugins/Runtime/Interchange/Source/Parsers/Fbx/Private/FbxHelper.cpp
+++ b/Engine/Plugins/Runtime/Interchange/Source/Parsers/Fbx/Private/FbxHelper.cpp
@@ -14,6 +14,62 @@ namespace UE
 	{
 		namespace Private
 		{
+			FString FFbxHelper::GetMeshName(FbxGeometryBase* Mesh)
+			{
+				FString MeshName = FFbxHelper::GetFbxObjectName(Mesh);
+				if (MeshName.IsEmpty())
+				{
+					if (Mesh->GetNodeCount() > 0)
+					{
+						if (Mesh->GetAttributeType() == FbxNodeAttribute::eMesh)
+						{
+							MeshName = TEXT("Mesh_");
+						}
+						else if (Mesh->GetAttributeType() == FbxNodeAttribute::eShape)
+						{
+							MeshName = TEXT("Shape_");
+						}
+						MeshName += FFbxHelper::GetFbxObjectName(Mesh->GetNode(0));
+					}
+					else
+					{
+						uint64 UniqueFbxObjectID = Mesh->GetUniqueID();
+						MeshName += GetUniqueIDString(UniqueFbxObjectID);
+					}
+				}
+				return MeshName;
+			}
+
+			FString FFbxHelper::GetMeshUniqueID(FbxGeometryBase* Mesh)
+			{
+				FString MeshUniqueID;
+				if (Mesh->GetAttributeType() == FbxNodeAttribute::eMesh)
+				{
+					MeshUniqueID = TEXT("\\Mesh\\");
+				}
+				else if (Mesh->GetAttributeType() == FbxNodeAttribute::eShape)
+				{
+					MeshUniqueID = TEXT("\\Shape\\");
+				}
+				FString MeshName = FFbxHelper::GetFbxObjectName(Mesh);
+				if (MeshName.IsEmpty())
+				{
+					if (Mesh->GetNodeCount() > 0)
+					{
+						MeshUniqueID += FFbxHelper::GetFbxNodeHierarchyName(Mesh->GetNode(0));
+					}
+					else
+					{
+						MeshUniqueID += GetMeshName(Mesh);
+					}
+				}
+				else
+				{
+					MeshUniqueID += MeshName;
+				}
+				return MeshUniqueID;
+			}
+
 			FString FFbxHelper::GetFbxObjectName(const FbxObject* Object)
 			{
 				if (!Object)
@@ -54,6 +110,13 @@ namespace UE
 				return UniqueID;
 			}
 
+			FString FFbxHelper::GetUniqueIDString(const uint64 UniqueID)
+			{
+				FStringFormatNamedArguments FormatArguments;
+				FormatArguments.Add(TEXT("UniqueID"), UniqueID);
+				return FString::Format(TEXT("{UniqueID}"), FormatArguments);
+			}
+#if 0
 			void FFbxHelper::FindSkeletalMeshes(FbxScene* SDKScene, TArray< TArray<FbxNode*> >& outSkelMeshArray, bool bCombineSkeletalMesh, bool bForceFindRigid)
 			{
 				TArray<FbxNode*> SkeletonArray;
@@ -965,7 +1028,7 @@ namespace UE
 					}
 				}
 			}
-
+#endif
 		}//ns Private
 	}//ns Interchange
 }//ns UE
diff --git a/Engine/Plugins/Runtime/Interchange/Source/Parsers/Fbx/Private/FbxHelper.h b/Engine/Plugins/Runtime/Interchange/Source/Parsers/Fbx/Private/FbxHelper.h
--- a/Engine/Plugins/Runtime/Interchange/Source/Parsers/Fbx/Private/FbxHelper.h
+++ b/Engine/Plugins/Runtime/Interchange/Source/Parsers/Fbx/Private/FbxHelper.h
@@ -11,18 +11,23 @@ namespace UE
 	{
 		namespace Private
 		{
+			class FFbxParser;
+
 			class FPayloadContextBase
 			{
 			public:
 				virtual ~FPayloadContextBase() {}
 				virtual FString GetPayloadType() const { return FString(); }
-				virtual bool FetchPayloadToFile(const FString& PayloadFilepath, TArray<FString>& JSonErrorMessages) { return false; }
+				virtual bool FetchPayloadToFile(FFbxParser& Parser, const FString& PayloadFilepath) { return false; }
 			};
 
 			struct FFbxHelper
 			{
 			public:
 
+				static FString GetMeshName(FbxGeometryBase* Mesh);
+				static FString GetMeshUniqueID(FbxGeometryBase* Mesh);
+
 				/**
 				 * Return the name of an FbxObject, return empty string if the object is null.
 				 */
@@ -35,6 +40,12 @@ namespace UE
 				 */
 				static FString GetFbxNodeHierarchyName(const FbxNode* Node);
 
+			protected:
+				static FString GetUniqueIDString(const uint64 UniqueID);
+
+#if 0
+				// This stuff is moved to FbxSkelMesh.h
+				// 
 				//Skeletalmesh helper
 				static void FindSkeletalMeshes(FbxScene* SDKScene, TArray< TArray<FbxNode*> >& outSkelMeshArray, bool bCombineSkeletalMesh, bool bForceFindRigid);
 				static void FindAllLODGroupNode(TArray<FbxNode*>& OutNodeInLod, FbxNode* NodeLodGroup, int32 LodIndex);
@@ -57,6 +68,7 @@ namespace UE
 				static void RecursiveFixSkeleton(FbxScene* SDKScene, FbxNode* Node, TArray<FbxNode*>& SkelMeshes, bool bImportNestedMeshes);
 
 				/* Skletalmesh private helper end */
+#endif
 			};
 		}//ns Private
 	}//ns Interchange
diff --git a/Engine/Plugins/Runtime/Interchange/Source/Parsers/Fbx/Private/FbxMaterial.cpp b/Engine/Plugins/Runtime/Interchange/Source/Parsers/Fbx/Private/FbxMaterial.cpp
--- a/Engine/Plugins/Runtime/Interchange/Source/Parsers/Fbx/Private/FbxMaterial.cpp
+++ b/Engine/Plugins/Runtime/Interchange/Source/Parsers/Fbx/Private/FbxMaterial.cpp
@@ -3,10 +3,12 @@
 #include "FbxMaterial.h"
 
 #include "CoreMinimal.h"
+#include "FbxAPI.h"
 #include "FbxConvert.h"
 #include "FbxHelper.h"
 #include "FbxInclude.h"
 #include "InterchangeMaterialNode.h"
+#include "InterchangeResultsContainer.h"
 #include "InterchangeSceneNode.h"
 #include "InterchangeTexture2DNode.h"
 #include "InterchangeTextureNode.h"
@@ -14,18 +16,21 @@
 #include "Nodes/InterchangeBaseNodeContainer.h"
 #include "Misc/Paths.h"
 
+#define LOCTEXT_NAMESPACE "InterchangeFbxMaterial"
+
 namespace UE
 {
 	namespace Interchange
 	{
 		namespace Private
 		{
-			UInterchangeMaterialNode* FFbxMaterial::CreateMaterialNode(UInterchangeBaseNodeContainer& NodeContainer, const FString& NodeUid, const FString& NodeName, TArray<FString>& JSonErrorMessages)
+			UInterchangeMaterialNode* FFbxMaterial::CreateMaterialNode(UInterchangeBaseNodeContainer& NodeContainer, const FString& NodeUid, const FString& NodeName)
 			{
 				UInterchangeMaterialNode* MaterialNode = NewObject<UInterchangeMaterialNode>(&NodeContainer, NAME_None);
 				if (!ensure(MaterialNode))
 				{
-					JSonErrorMessages.Add(TEXT("{\"Msg\" : {\"Type\" : \"Error\",\n\"Msg\" : \"Cannot allocate a node when importing fbx\"}}"));
+					UInterchangeResultError_Generic* Message = Parser.AddMessage<UInterchangeResultError_Generic>();
+					Message->Text = LOCTEXT("CannotAllocateNode", "Cannot allocate a node when importing FBX.");
 					return nullptr;
 				}
 				// Creating a UMaterialInterface
@@ -35,13 +40,14 @@ namespace UE
 				return MaterialNode;
 			}
 
-			UInterchangeTexture2DNode* FFbxMaterial::CreateTexture2DNode(UInterchangeBaseNodeContainer& NodeContainer, const FString& NodeUid, const FString& TextureFilePath, TArray<FString>& JSonErrorMessages)
+			UInterchangeTexture2DNode* FFbxMaterial::CreateTexture2DNode(UInterchangeBaseNodeContainer& NodeContainer, const FString& NodeUid, const FString& TextureFilePath)
 			{
 				FString DisplayLabel = FPaths::GetBaseFilename(TextureFilePath);
 				UInterchangeTexture2DNode* TextureNode = NewObject<UInterchangeTexture2DNode>(&NodeContainer, NAME_None);
 				if (!ensure(TextureNode))
 				{
-					JSonErrorMessages.Add(TEXT("{\"Msg\" : {\"Type\" : \"Error\",\n\"Msg\" : \"Cannot allocate a node when importing fbx\"}}"));
+					UInterchangeResultError_Generic* Message = Parser.AddMessage<UInterchangeResultError_Generic>();
+					Message->Text = LOCTEXT("CannotAllocateNode", "Cannot allocate a node when importing FBX.");
 					return nullptr;
 				}
 				// Creating a UTexture2D
@@ -54,7 +60,7 @@ namespace UE
 				return TextureNode;
 			}
 
-			UInterchangeMaterialNode* FFbxMaterial::AddNodeMaterial(FbxSurfaceMaterial* SurfaceMaterial, UInterchangeBaseNodeContainer& NodeContainer, TArray<FString>& JSonErrorMessages)
+			UInterchangeMaterialNode* FFbxMaterial::AddNodeMaterial(FbxSurfaceMaterial* SurfaceMaterial, UInterchangeBaseNodeContainer& NodeContainer)
 			{
 				//Create a material node
 				FString MaterialName = FFbxHelper::GetFbxObjectName(SurfaceMaterial);
@@ -62,14 +68,19 @@ namespace UE
 				UInterchangeMaterialNode* MaterialNode = Cast<UInterchangeMaterialNode>(NodeContainer.GetNode(NodeUid));
 				if (!MaterialNode)
 				{
-					MaterialNode = CreateMaterialNode(NodeContainer, NodeUid, MaterialName, JSonErrorMessages);
+					MaterialNode = CreateMaterialNode(NodeContainer, NodeUid, MaterialName);
 					if (MaterialNode == nullptr)
 					{
-						JSonErrorMessages.Add(TEXT("{\"Msg\" : {\"Type\" : \"Error\",\n\"Msg\" : \"Cannot create fbx material (" + MaterialName + TEXT(")\"}}")));
+						FFormatNamedArguments Args
+						{
+							{ TEXT("MaterialName"), FText::FromString(MaterialName) }
+						};
+						UInterchangeResultError_Generic* Message = Parser.AddMessage<UInterchangeResultError_Generic>();
+						Message->Text = FText::Format(LOCTEXT("CannotCreateFBXMaterial", "Cannot create FBX material '{MaterialName}'."), Args);
 						return nullptr;
 					}
 
-					auto SetMaterialParameter = [&NodeContainer, &MaterialNode, &SurfaceMaterial, &JSonErrorMessages, &MaterialName](const char* FbxMaterialProperty, EInterchangeMaterialNodeParameterName MaterialParameterName)
+					auto SetMaterialParameter = [this, &NodeContainer, &MaterialNode, &SurfaceMaterial, &MaterialName](const char* FbxMaterialProperty, EInterchangeMaterialNodeParameterName MaterialParameterName)
 					{
 						bool bSetMaterial = false;
 						FbxProperty FbxProperty = SurfaceMaterial->FindProperty(FbxMaterialProperty);
@@ -81,7 +92,12 @@ namespace UE
 							bool bFoundValidTexture = false;
 							if (UnsupportedTextureCount > 0)
 							{
-								JSonErrorMessages.Add(TEXT("{\"Msg\" : {\"Type\" : \"Warning\",\n\"Msg\" : \"Layered or procedural Textures are not supported (material" + MaterialName + TEXT(")\"}}")));
+								FFormatNamedArguments Args
+								{
+									{ TEXT("MaterialName"), FText::FromString(MaterialName) }
+								};
+								UInterchangeResultWarning_Generic* Message = Parser.AddMessage<UInterchangeResultWarning_Generic>();
+								Message->Text = FText::Format(LOCTEXT("TextureTypeNotSupported", "Layered or procedural textures are not supported (material '{MaterialName}')."), Args);
 							}
 							else if (TextureCount > 0)
 							{
@@ -100,7 +116,7 @@ namespace UE
 									UInterchangeTexture2DNode* TextureNode = Cast<UInterchangeTexture2DNode>(NodeContainer.GetNode(NodeUid));
 									if (!TextureNode)
 									{
-										TextureNode = CreateTexture2DNode(NodeContainer, NodeUid, TextureFilename, JsonErrorMessage);
+										TextureNode = CreateTexture2DNode(NodeContainer, NodeUid, TextureFilename);
 									}
 									// add/find UVSet and set it to the texture, we pass index 0 here, I think pipeline should be able to get the UVIndex channel from the name
 									// and modify the parameter to set the correct value.
@@ -174,7 +190,7 @@ namespace UE
 				return MaterialNode;
 			}
 
-			void FFbxMaterial::AddAllTextures(FbxScene* SDKScene, UInterchangeBaseNodeContainer& NodeContainer, TArray<FString>& JSonErrorMessages)
+			void FFbxMaterial::AddAllTextures(FbxScene* SDKScene, UInterchangeBaseNodeContainer& NodeContainer)
 			{
 				int32 TextureCount = SDKScene->GetSrcObjectCount<FbxFileTexture>();
 				for (int32 TextureIndex = 0; TextureIndex < TextureCount; ++TextureIndex)
@@ -192,32 +208,34 @@ namespace UE
 					UInterchangeTexture2DNode* TextureNode = Cast<UInterchangeTexture2DNode>(NodeContainer.GetNode(NodeUid));
 					if (!TextureNode)
 					{
-						TextureNode = CreateTexture2DNode(NodeContainer, NodeUid, TextureFilename, JsonErrorMessage);
+						TextureNode = CreateTexture2DNode(NodeContainer, NodeUid, TextureFilename);
 					}
 				}
 			}
 			
-			void FFbxMaterial::AddAllNodeMaterials(UInterchangeSceneNode* SceneNode, FbxNode* ParentFbxNode, UInterchangeBaseNodeContainer& NodeContainer, TArray<FString>& JSonErrorMessages)
+			void FFbxMaterial::AddAllNodeMaterials(UInterchangeSceneNode* SceneNode, FbxNode* ParentFbxNode, UInterchangeBaseNodeContainer& NodeContainer)
 			{
 				int32 MaterialCount = ParentFbxNode->GetMaterialCount();
 				for (int32 MaterialIndex = 0; MaterialIndex < MaterialCount; ++MaterialIndex)
 				{
 					FbxSurfaceMaterial* SurfaceMaterial = ParentFbxNode->GetMaterial(MaterialIndex);
-					UInterchangeMaterialNode* MaterialNode = AddNodeMaterial(SurfaceMaterial, NodeContainer, JSonErrorMessages);
+					UInterchangeMaterialNode* MaterialNode = AddNodeMaterial(SurfaceMaterial, NodeContainer);
 					//The dependencies order is important because mesh will use index in that order to determine material use by a face
 					SceneNode->AddMaterialDependencyUid(MaterialNode->GetUniqueID());
 				}
 			}
 
-			void FFbxMaterial::AddAllMaterials(FbxScene* SDKScene, UInterchangeBaseNodeContainer& NodeContainer, TArray<FString>& JSonErrorMessages)
+			void FFbxMaterial::AddAllMaterials(FbxScene* SDKScene, UInterchangeBaseNodeContainer& NodeContainer)
 			{
 				int32 MaterialCount = SDKScene->GetMaterialCount();
 				for (int32 MaterialIndex = 0; MaterialIndex < MaterialCount; ++MaterialIndex)
 				{
 					FbxSurfaceMaterial* SurfaceMaterial = SDKScene->GetMaterial(MaterialIndex);
-					AddNodeMaterial(SurfaceMaterial, NodeContainer, JSonErrorMessages);
+					AddNodeMaterial(SurfaceMaterial, NodeContainer);
 				}
 			}
 		} //ns Private
 	} //ns Interchange
 }//ns UE
+
+#undef LOCTEXT_NAMESPACE
diff --git a/Engine/Plugins/Runtime/Interchange/Source/Parsers/Fbx/Private/FbxMaterial.h b/Engine/Plugins/Runtime/Interchange/Source/Parsers/Fbx/Private/FbxMaterial.h
--- a/Engine/Plugins/Runtime/Interchange/Source/Parsers/Fbx/Private/FbxMaterial.h
+++ b/Engine/Plugins/Runtime/Interchange/Source/Parsers/Fbx/Private/FbxMaterial.h
@@ -18,22 +18,28 @@ namespace UE
 	{
 		namespace Private
 		{
+			class FFbxParser;
+
 			class FFbxMaterial
 			{
 			public:
+				explicit FFbxMaterial(FFbxParser& InParser)
+					: Parser(InParser)
+				{}
+
 				/**
 				 * Create a UInterchangeFextureNode and add it to the NodeContainer for all texture of type FbxFileTexture the fbx file contain.
 				 *
 				 * @note - Any node that already exist in the NodeContainer will not be created or modify.
 				 */
-				static void AddAllTextures(FbxScene* SDKScene, UInterchangeBaseNodeContainer& NodeContainer, TArray<FString>& JSonErrorMessages);
+				void AddAllTextures(FbxScene* SDKScene, UInterchangeBaseNodeContainer& NodeContainer);
 				
 				/**
 				 * Create a UInterchangeMaterialNode and add it to the NodeContainer for all material of type FbxSurfaceMaterial the fbx file contain.
 				 * 
 				 * @note - Any node that already exist in the NodeContainer will not be created or modify.
 				 */
-				static void AddAllMaterials(FbxScene* SDKScene, UInterchangeBaseNodeContainer& NodeContainer, TArray<FString>& JSonErrorMessages);
+				void AddAllMaterials(FbxScene* SDKScene, UInterchangeBaseNodeContainer& NodeContainer);
 
 				/**
 				 * Create a UInterchangeMaterialNode and add it to the NodeContainer for all material of type FbxSurfaceMaterial the fbx ParentFbxNode contain.
@@ -41,13 +47,16 @@ namespace UE
 				 * 
 				 * @note - Any material node that already exist in the NodeContainer will simply be add has a dependency.
 				 */
-				static void AddAllNodeMaterials(UInterchangeSceneNode* SceneNode, FbxNode* ParentFbxNode, UInterchangeBaseNodeContainer& NodeContainer, TArray<FString>& JSonErrorMessages);
+				void AddAllNodeMaterials(UInterchangeSceneNode* SceneNode, FbxNode* ParentFbxNode, UInterchangeBaseNodeContainer& NodeContainer);
+
 			protected:
-				static UInterchangeMaterialNode* CreateMaterialNode(UInterchangeBaseNodeContainer& NodeContainer, const FString& NodeUID, const FString& NodeName, TArray<FString>& JSonErrorMessages);
-				static UInterchangeTexture2DNode* CreateTexture2DNode(UInterchangeBaseNodeContainer& NodeContainer, const FString& NodeUID, const FString& TextureFilePath, TArray<FString>& JSonErrorMessages);
+				UInterchangeMaterialNode* CreateMaterialNode(UInterchangeBaseNodeContainer& NodeContainer, const FString& NodeUID, const FString& NodeName);
+				UInterchangeTexture2DNode* CreateTexture2DNode(UInterchangeBaseNodeContainer& NodeContainer, const FString& NodeUID, const FString& TextureFilePath);
 
 			private:
-				static UInterchangeMaterialNode* AddNodeMaterial(FbxSurfaceMaterial* SurfaceMaterial, UInterchangeBaseNodeContainer& NodeContainer, TArray<FString>& JSonErrorMessages);
+				UInterchangeMaterialNode* AddNodeMaterial(FbxSurfaceMaterial* SurfaceMaterial, UInterchangeBaseNodeContainer& NodeContainer);
+
+				FFbxParser& Parser;
 			};
 		}//ns Private
 	}//ns Interchange
diff --git a/Engine/Plugins/Runtime/Interchange/Source/Parsers/Fbx/Private/FbxMesh.cpp b/Engine/Plugins/Runtime/Interchange/Source/Parsers/Fbx/Private/FbxMesh.cpp
--- a/Engine/Plugins/Runtime/Interchange/Source/Parsers/Fbx/Private/FbxMesh.cpp
+++ b/Engine/Plugins/Runtime/Interchange/Source/Parsers/Fbx/Private/FbxMesh.cpp
@@ -3,10 +3,13 @@
 #include "FbxMesh.h"
 
 #include "CoreMinimal.h"
+#include "FbxAPI.h"
 #include "FbxConvert.h"
 #include "FbxHelper.h"
 #include "FbxInclude.h"
+#include "Fbx/InterchangeFbxMessages.h"
 #include "InterchangeMeshNode.h"
+#include "InterchangeResultsContainer.h"
 #include "MeshDescription.h"
 #include "Misc/FileHelper.h"
 #include "Nodes/InterchangeBaseNodeContainer.h"
@@ -14,6 +17,8 @@
 #include "SkeletalMeshAttributes.h"
 #include "StaticMeshAttributes.h"
 
+#define LOCTEXT_NAMESPACE "InterchangeFbxMesh"
+
 namespace UE
 {
 	namespace Interchange
@@ -24,8 +29,9 @@ namespace UE
 			struct FFBXUVs
 			{
 				// constructor
-				FFBXUVs(FbxMesh* Mesh)
-					: UniqueUVCount(0)
+				FFBXUVs(FMeshDescriptionImporter& InMeshDescriptionImporter, FbxMesh* Mesh)
+					: MeshDescriptionImporter(InMeshDescriptionImporter)
+					, UniqueUVCount(0)
 				{
 					check(Mesh);
 
@@ -140,8 +146,8 @@ namespace UE
 
 					if (UniqueUVCount > MAX_MESH_TEXTURE_COORDS_MD)
 					{
-						//TODO Error Message
-						//FbxImporter->AddTokenizedErrorMessage(FTokenizedMessage::Create(EMessageSeverity::Warning, FText::Format(LOCTEXT("Error_TooMuchUVChannel", "Reached the maximum number of UV Channels for a Static Mesh({0}) - discarding {1} UV Channels"), FText::AsNumber(MAX_MESH_TEXTURE_COORDS_MD), FText::AsNumber(UniqueUVCount - MAX_MESH_TEXTURE_COORDS_MD))), FFbxErrors::Generic_Mesh_TooMuchUVChannels);
+						UInterchangeResultMeshWarning_TooManyUVs* Message = MeshDescriptionImporter.AddMessage<UInterchangeResultMeshWarning_TooManyUVs>(Mesh);
+						Message->ExcessUVs = UniqueUVCount - MAX_MESH_TEXTURE_COORDS_MD;
 					}
 
 					UniqueUVCount = FMath::Min<int32>(UniqueUVCount, MAX_MESH_TEXTURE_COORDS_MD);
@@ -178,6 +184,7 @@ namespace UE
 					UVMappingMode.Empty();
 				}
 
+				FMeshDescriptionImporter& MeshDescriptionImporter;
 				TArray<FString> UVSets;
 				TArray<FbxLayerElementUV const*> LayerElementUV;
 				TArray<FbxLayerElement::EReferenceMode> UVReferenceMode;
@@ -185,8 +192,9 @@ namespace UE
 				int32 UniqueUVCount;
 			};
 
-			FMeshDescriptionImporter::FMeshDescriptionImporter(FMeshDescription* InMeshDescription, FbxScene* InSDKScene, FbxGeometryConverter* InSDKGeometryConverter)
-				: MeshDescription(InMeshDescription)
+			FMeshDescriptionImporter::FMeshDescriptionImporter(FFbxParser& InParser, FMeshDescription* InMeshDescription, FbxScene* InSDKScene, FbxGeometryConverter* InSDKGeometryConverter)
+				: Parser(InParser)
+				, MeshDescription(InMeshDescription)
 				, SDKScene(InSDKScene)
 				, SDKGeometryConverter(InSDKGeometryConverter)
 			{
@@ -265,8 +273,9 @@ namespace UE
 						FVertexID AddedVertexId = MeshDescription->CreateVertex();
 						if (AddedVertexId.GetValue() != RealVertexIndex)
 						{
-							//TODO add an error log
-							//AddTokenizedErrorMessage(FTokenizedMessage::Create(EMessageSeverity::Error, FText::Format(LOCTEXT("Error_CannotCreateVertex", "Cannot create valid vertex for mesh '{0}'"), FText::FromString(Mesh->GetName()))), FFbxErrors::StaticMesh_BuildError);
+							UInterchangeResultMeshError_Generic* Message = AddMessage<UInterchangeResultMeshError_Generic>(Shape);
+							Message->Text = LOCTEXT("CantCreateVertex", "Cannot create valid vertex for the mesh '{MeshName}'.");
+
 							return false;
 						}
 						//Add the delta position, so we do not have to recompute it later
@@ -286,6 +295,9 @@ namespace UE
 					return false;
 				}
 
+				UInterchangeResultMeshWarning_Generic* TestMessage = AddMessage<UInterchangeResultMeshWarning_Generic>(Mesh);
+				TestMessage->Text = LOCTEXT("TestMessage", "Test message for mesh '{MeshName}'.");
+
 				FStaticMeshAttributes Attributes(*MeshDescription);
 				Attributes.Register();
 
@@ -293,11 +305,12 @@ namespace UE
 				FbxLayer* BaseLayer = Mesh->GetLayer(0);
 				if (BaseLayer == NULL)
 				{
-					//TODO add error message
+					UInterchangeResultMeshError_Generic* Message = AddMessage<UInterchangeResultMeshError_Generic>(Mesh);
+					Message->Text = LOCTEXT("CantGetLayer", "Unable to get FBX base layer for the mesh '{MeshName}'.");
 					return false;
 				}
 
-				FFBXUVs FBXUVs(Mesh);
+				FFBXUVs FBXUVs(*this, Mesh);
 
 				FbxNode* MeshNode = Mesh->GetNodeCount() > 0 ? Mesh->GetNode() : nullptr; //Get the first node using the mesh to setup default asset materials
 				//
@@ -339,7 +352,9 @@ namespace UE
 					}
 					else
 					{
-						//TODO add an error message
+						UInterchangeResultMeshError_Generic* Message = AddMessage<UInterchangeResultMeshError_Generic>(Mesh);
+						Message->Text = LOCTEXT("TriangulationError", "Unable to triangulate the mesh '{MeshName}'.");
+
 						MeshDescription->Empty();
 						return false;
 					}
@@ -485,8 +500,9 @@ namespace UE
 
 					if (PolygonCount == 0)
 					{
-						//TODO add an error log
-						//AddTokenizedErrorMessage(FTokenizedMessage::Create(EMessageSeverity::Error, FText::Format(LOCTEXT("Error_NoPolygonFoundInMesh", "No polygon were found on mesh  '{0}'"), FText::FromString(Mesh->GetName()))), FFbxErrors::StaticMesh_NoTriangles);
+						UInterchangeResultMeshError_Generic* Message = AddMessage<UInterchangeResultMeshError_Generic>(Mesh);
+						Message->Text = LOCTEXT("NoPolygons", "No polygons were found in the mesh '{MeshName}'.");
+
 						return false;
 					}
 
@@ -546,8 +562,9 @@ namespace UE
 						VertexPositions[AddedVertexId] = VertexPosition;
 						if (AddedVertexId.GetValue() != RealVertexIndex)
 						{
-							//TODO add an error log
-							//AddTokenizedErrorMessage(FTokenizedMessage::Create(EMessageSeverity::Error, FText::Format(LOCTEXT("Error_CannotCreateVertex", "Cannot create valid vertex for mesh '{0}'"), FText::FromString(Mesh->GetName()))), FFbxErrors::StaticMesh_BuildError);
+							UInterchangeResultMeshError_Generic* Message = AddMessage<UInterchangeResultMeshError_Generic>(Mesh);
+							Message->Text = LOCTEXT("CannotCreateVertex", "Cannot create valid vertex for mesh '{MeshName}'.");
+
 							return false;
 						}
 					}
@@ -586,8 +603,8 @@ namespace UE
 							// Skip invalid edges, i.e. one of the ends is invalid, or degenerated ones
 							if (EdgeStartVertexIndex == -1 || EdgeEndVertexIndex == -1 || EdgeStartVertexIndex == EdgeEndVertexIndex)
 							{
-								//TODO log a user error
-								//UE_LOG(LogFbx, Warning, TEXT("Skipping invalid edge on mesh %s"), *FString(Mesh->GetName()));
+								UInterchangeResultMeshWarning_Generic* Message = AddMessage<UInterchangeResultMeshWarning_Generic>(Mesh);
+								Message->Text = LOCTEXT("SkippingEdge", "Skipping invalid edge on mesh '{MeshName}'.");
 								continue;
 							}
 							FVertexID EdgeVertexStart(EdgeStartVertexIndex + VertexOffset);
@@ -678,8 +695,9 @@ namespace UE
 
 								if (AddedVertexInstanceId.GetValue() != VertexInstanceIndex)
 								{
-									//TODO add an error log
-									//AddTokenizedErrorMessage(FTokenizedMessage::Create(EMessageSeverity::Error, FText::Format(LOCTEXT("Error_CannotCreateVertexInstance", "Cannot create valid vertex instance for mesh '{0}'"), FText::FromString(Mesh->GetName()))), FFbxErrors::StaticMesh_BuildError);
+									UInterchangeResultMeshError_Generic* Message = AddMessage<UInterchangeResultMeshError_Generic>(Mesh);
+									Message->Text = LOCTEXT("CannotCreateVertexInstance", "Cannot create valid vertex instance for the mesh '{MeshName}'.");
+
 									return false;
 								}
 
@@ -1102,23 +1120,29 @@ namespace UE
 			//////////////////////////////////////////////////////////////////////////
 			/// FMeshPayloadContext implementation
 
-			bool FMeshPayloadContext::FetchPayloadToFile(const FString& PayloadFilepath, TArray<FString>& JSonErrorMessages)
+			bool FMeshPayloadContext::FetchPayloadToFile(FFbxParser& Parser, const FString& PayloadFilepath)
 			{
 				if (!ensure(SDKScene != nullptr))
 				{
-					JSonErrorMessages.Add(TEXT("{\"Msg\" : {\"Type\" : \"Error\",\n\"Msg\" : \"Cannot fetch fbx mesh payload because the fbx scene is null\"}}"));
+					UInterchangeResultError_Generic* Message = Parser.AddMessage<UInterchangeResultError_Generic>();
+					Message->InterchangeKey = FFbxHelper::GetMeshUniqueID(Mesh);
+					Message->Text = LOCTEXT("FBXSceneNull", "Cannot fetch FBX mesh payload because the FBX scene is null.");
 					return false;
 				}
 
 				if (!ensure(Mesh != nullptr))
 				{
-					JSonErrorMessages.Add(TEXT("{\"Msg\" : {\"Type\" : \"Error\",\n\"Msg\" : \"Cannot fetch fbx mesh payload because the fbx mesh is null\"}}"));
+					UInterchangeResultError_Generic* Message = Parser.AddMessage<UInterchangeResultError_Generic>();
+					Message->InterchangeKey = FFbxHelper::GetMeshUniqueID(Mesh);
+					Message->Text = LOCTEXT("FBXMeshNull", "Cannot fetch FBX mesh payload because the FBX mesh is null.");
 					return false;
 				}
 
 				if (!ensure(SDKGeometryConverter != nullptr))
 				{
-					JSonErrorMessages.Add(TEXT("{\"Msg\" : {\"Type\" : \"Error\",\n\"Msg\" : \"Cannot fetch fbx mesh payload because the fbx scene geometry converter is null\"}}"));
+					UInterchangeResultError_Generic* Message = Parser.AddMessage<UInterchangeResultError_Generic>();
+					Message->InterchangeKey = FFbxHelper::GetMeshUniqueID(Mesh);
+					Message->Text = LOCTEXT("FBXConverterNull", "Cannot fetch FBX mesh payload because the FBX geometry converter is null.");
 					return false;
 				}
 
@@ -1131,10 +1155,12 @@ namespace UE
 				{
 					FSkeletalMeshAttributes SkeletalMeshAttribute(MeshDescription);
 					SkeletalMeshAttribute.Register();
-					FMeshDescriptionImporter MeshDescriptionImporter(&MeshDescription, SDKScene, SDKGeometryConverter);
+					FMeshDescriptionImporter MeshDescriptionImporter(Parser, &MeshDescription, SDKScene, SDKGeometryConverter);
 					if (!MeshDescriptionImporter.FillSkinnedMeshDescriptionFromFbxMesh(Mesh, JointUniqueNames))
 					{
-						JSonErrorMessages.Add(TEXT("{\"Msg\" : {\"Type\" : \"Error\",\n\"Msg\" : \"Cannot fetch skinned mesh payload because there was an error when creating the FMeshDescription\"}}"));
+						UInterchangeResultError_Generic* Message = Parser.AddMessage<UInterchangeResultError_Generic>();
+						Message->InterchangeKey = FFbxHelper::GetMeshUniqueID(Mesh);
+						Message->Text = LOCTEXT("MeshDescriptionError", "Cannot fetch skinned mesh payload because there was an error when creating the MeshDescription.");
 						return false;
 					}
 				}
@@ -1142,10 +1168,12 @@ namespace UE
 				{
 					FStaticMeshAttributes StaticMeshAttribute(MeshDescription);
 					StaticMeshAttribute.Register();
-					FMeshDescriptionImporter MeshDescriptionImporter(&MeshDescription, SDKScene, SDKGeometryConverter);
+					FMeshDescriptionImporter MeshDescriptionImporter(Parser, &MeshDescription, SDKScene, SDKGeometryConverter);
 					if (!MeshDescriptionImporter.FillStaticMeshDescriptionFromFbxMesh(Mesh))
 					{
-						JSonErrorMessages.Add(TEXT("{\"Msg\" : {\"Type\" : \"Error\",\n\"Msg\" : \"Cannot fetch static mesh payload because there was an error when creating the FMeshDescription\"}}"));
+						UInterchangeResultError_Generic* Message = Parser.AddMessage<UInterchangeResultError_Generic>();
+						Message->InterchangeKey = FFbxHelper::GetMeshUniqueID(Mesh);
+						Message->Text = LOCTEXT("MeshDescriptionError", "Cannot fetch static mesh payload because there was an error when creating the MeshDescription.");
 						return false;
 					}
 				}
@@ -1172,7 +1200,7 @@ namespace UE
 			//////////////////////////////////////////////////////////////////////////
 			/// FShapePayloadContext implementation
 
-			bool FShapePayloadContext::FetchPayloadToFile(const FString& PayloadFilepath, TArray<FString>& JSonErrorMessages)
+			bool FShapePayloadContext::FetchPayloadToFile(FFbxParser& Parser, const FString& PayloadFilepath)
 			{
 				if (!ensure(Shape))
 				{
@@ -1183,10 +1211,13 @@ namespace UE
 				FMeshDescription ShapeMeshDescription;
 				FStaticMeshAttributes StaticMeshAttribute(ShapeMeshDescription);
 				StaticMeshAttribute.Register();
-				FMeshDescriptionImporter MeshDescriptionImporter(&ShapeMeshDescription, SDKScene, SDKGeometryConverter);
+				FMeshDescriptionImporter MeshDescriptionImporter(Parser, &ShapeMeshDescription, SDKScene, SDKGeometryConverter);
 				if (!MeshDescriptionImporter.FillMeshDescriptionFromFbxShape(Shape))
 				{
-					//Todo log an error
+					UInterchangeResultError_Generic* Message = Parser.AddMessage<UInterchangeResultError_Generic>();
+					Message->InterchangeKey = FFbxHelper::GetMeshUniqueID(Shape);
+					Message->Text = LOCTEXT("MeshDescriptionBlendShapeError", "Unable to create MeshDescription from FBX blend shape.");
+
 					return false;
 				}
 
@@ -1204,7 +1235,7 @@ namespace UE
 
 			//////////////////////////////////////////////////////////////////////////
 			/// FFbxMesh implementation
-			
+#if 0
 			FString FFbxMesh::GetUniqueIDString(const uint64 UniqueID)
 			{
 				FStringFormatNamedArguments FormatArguments;
@@ -1267,8 +1298,9 @@ namespace UE
 				}
 				return MeshUniqueID;
 			}
+#endif
 
-			void FFbxMesh::ExtractSkinnedMeshNodeJoints(FbxScene* SDKScene, FbxMesh* Mesh, UInterchangeMeshNode* MeshNode, TArray<FString>& JSonErrorMessages)
+			void FFbxMesh::ExtractSkinnedMeshNodeJoints(FbxScene* SDKScene, FbxMesh* Mesh, UInterchangeMeshNode* MeshNode)
 			{
 				TArray<FString> JointNodeUniqueIDs;
 				const int32 SkinDeformerCount = Mesh->GetDeformerCount(FbxDeformer::eSkin);
@@ -1303,7 +1335,7 @@ namespace UE
 				}
 			}
 
-			void FFbxMesh::AddAllMeshes(FbxScene* SDKScene, FbxGeometryConverter* SDKGeometryConverter, UInterchangeBaseNodeContainer& NodeContainer, TArray<FString>& JSonErrorMessages, TMap<FString, TSharedPtr<FPayloadContextBase>>& PayloadContexts)
+			void FFbxMesh::AddAllMeshes(FbxScene* SDKScene, FbxGeometryConverter* SDKGeometryConverter, UInterchangeBaseNodeContainer& NodeContainer, TMap<FString, TSharedPtr<FPayloadContextBase>>& PayloadContexts)
 			{
 				int32 GeometryCount = SDKScene->GetGeometryCount();
 				for (int32 GeometryIndex = 0; GeometryIndex < GeometryCount; ++GeometryIndex)
@@ -1318,17 +1350,17 @@ namespace UE
 					{
 						continue;
 					}
-					FString MeshName = GetMeshName(Mesh);
-					FString MeshUniqueID = GetMeshUniqueID(Mesh);
+					FString MeshName = FFbxHelper::GetMeshName(Mesh);
+					FString MeshUniqueID = FFbxHelper::GetMeshUniqueID(Mesh);
 					UInterchangeMeshNode* MeshNode = Cast<UInterchangeMeshNode>(NodeContainer.GetNode(MeshUniqueID));
 					if (!MeshNode)
 					{
-						MeshNode = CreateMeshNode(NodeContainer, MeshName, MeshUniqueID, JSonErrorMessages);
+						MeshNode = CreateMeshNode(NodeContainer, MeshName, MeshUniqueID);
 						if (Geometry->GetDeformerCount(FbxDeformer::eSkin) > 0)
 						{
 							//Set the skinned mesh attribute
 							MeshNode->SetSkinnedMesh(true);
-							ExtractSkinnedMeshNodeJoints(SDKScene, Mesh, MeshNode, JSonErrorMessages);
+							ExtractSkinnedMeshNodeJoints(SDKScene, Mesh, MeshNode);
 						}
 						Mesh->ComputeBBox();
 						const int32 MeshVertexCount = Mesh->GetControlPointsCount();
@@ -1421,12 +1453,12 @@ namespace UE
 										}
 										ensure(!ShapeNameToFbxShape.Contains(ShapeName));
 										ShapeNameToFbxShape.Add(ShapeName, Shape);
-										FString ShapeAttributeName = GetMeshName(Shape);
-										FString ShapeUniqueID = GetMeshUniqueID(Shape);
+										FString ShapeAttributeName = FFbxHelper::GetMeshName(Shape);
+										FString ShapeUniqueID = FFbxHelper::GetMeshUniqueID(Shape);
 										UInterchangeMeshNode* ShapeNode = Cast<UInterchangeMeshNode>(NodeContainer.GetNode(ShapeUniqueID));
 										if (!ShapeNode)
 										{
-											ShapeNode = CreateMeshNode(NodeContainer, ShapeAttributeName, ShapeUniqueID, JSonErrorMessages);
+											ShapeNode = CreateMeshNode(NodeContainer, ShapeAttributeName, ShapeUniqueID);
 											const bool bIsBlendShape = true;
 											ShapeNode->SetBlendShape(bIsBlendShape);
 											ShapeNode->SetBlendShapeName(ShapeName);
@@ -1451,14 +1483,16 @@ namespace UE
 				}
 			}
 
-			UInterchangeMeshNode* FFbxMesh::CreateMeshNode(UInterchangeBaseNodeContainer& NodeContainer, const FString& NodeName, const FString& NodeUniqueID, TArray<FString>& JSonErrorMessages)
+			UInterchangeMeshNode* FFbxMesh::CreateMeshNode(UInterchangeBaseNodeContainer& NodeContainer, const FString& NodeName, const FString& NodeUniqueID)
 			{
 				FString DisplayLabel(NodeName);
 				FString NodeUid(NodeUniqueID);
 				UInterchangeMeshNode* MeshNode = NewObject<UInterchangeMeshNode>(&NodeContainer, NAME_None);
 				if (!ensure(MeshNode))
 				{
-					JSonErrorMessages.Add(TEXT("{\"Msg\" : {\"Type\" : \"Error\",\n\"Msg\" : \"Cannot allocate a mesh node when importing fbx\"}}"));
+					UInterchangeResultError_Generic* Message = Parser.AddMessage<UInterchangeResultError_Generic>();
+					Message->InterchangeKey = NodeUniqueID;
+					Message->Text = LOCTEXT("NodeAllocationError", "Mesh node allocation failed when importing FBX.");
 					return nullptr;
 				}
 				// Creating a UMaterialInterface
@@ -1469,3 +1503,5 @@ namespace UE
 		} //ns Private
 	} //ns Interchange
 }//ns UE
+
+#undef LOCTEXT_NAMESPACE
diff --git a/Engine/Plugins/Runtime/Interchange/Source/Parsers/Fbx/Private/FbxMesh.h b/Engine/Plugins/Runtime/Interchange/Source/Parsers/Fbx/Private/FbxMesh.h
--- a/Engine/Plugins/Runtime/Interchange/Source/Parsers/Fbx/Private/FbxMesh.h
+++ b/Engine/Plugins/Runtime/Interchange/Source/Parsers/Fbx/Private/FbxMesh.h
@@ -3,6 +3,7 @@
 #pragma once
 
 #include "CoreMinimal.h"
+#include "FbxAPI.h"
 #include "FbxHelper.h"
 #include "FbxInclude.h"
 
@@ -17,10 +18,12 @@ namespace UE
 	{
 		namespace Private
 		{
+			class FFbxMesh;
+
 			class FMeshDescriptionImporter
 			{
 			public:
-				FMeshDescriptionImporter(FMeshDescription* InMeshDescription, FbxScene* InSDKScene, FbxGeometryConverter* InSDKGeometryConverter);
+				FMeshDescriptionImporter(FFbxParser& InParser, FMeshDescription* InMeshDescription, FbxScene* InSDKScene, FbxGeometryConverter* InSDKGeometryConverter);
 				
 				/*
 				 * Fill the mesh description using the Mesh parameter.
@@ -37,6 +40,18 @@ namespace UE
 				 */
 				bool FillMeshDescriptionFromFbxShape(FbxShape* Shape);
 
+				/**
+				 * Add messages to the message log
+				 */
+				template <typename T>
+				T* AddMessage(FbxGeometryBase* FbxNode) const
+				{
+					T* Item = Parser.AddMessage<T>();
+					Item->MeshName = FFbxHelper::GetMeshName(FbxNode);
+					Item->InterchangeKey = FFbxHelper::GetMeshUniqueID(FbxNode);
+					return Item;
+				}
+
 			private:
 				
 				enum class EMeshType : uint8
@@ -60,6 +75,8 @@ namespace UE
 					);
 					return (Basis.Determinant() < 0) ? -1.0f : +1.0f;
 				}
+
+				FFbxParser& Parser;
 				FMeshDescription* MeshDescription;
 				FbxScene* SDKScene;
 				FbxGeometryConverter* SDKGeometryConverter;
@@ -71,7 +88,7 @@ namespace UE
 			public:
 				virtual ~FMeshPayloadContext() {}
 				virtual FString GetPayloadType() const override { return TEXT("Mesh-PayloadContext"); }
-				virtual bool FetchPayloadToFile(const FString& PayloadFilepath, TArray<FString>& JSonErrorMessages) override;
+				virtual bool FetchPayloadToFile(FFbxParser& Parser, const FString& PayloadFilepath) override;
 				FbxMesh* Mesh = nullptr;
 				FbxScene* SDKScene = nullptr;
 				FbxGeometryConverter* SDKGeometryConverter = nullptr;
@@ -82,7 +99,7 @@ namespace UE
 			public:
 				virtual ~FShapePayloadContext() {}
 				virtual FString GetPayloadType() const override { return TEXT("Shape-PayloadContext"); }
-				virtual bool FetchPayloadToFile(const FString& PayloadFilepath, TArray<FString>& JSonErrorMessages) override;
+				virtual bool FetchPayloadToFile(FFbxParser& Parser, const FString& PayloadFilepath) override;
 				FbxShape* Shape = nullptr;
 				FbxScene* SDKScene = nullptr;
 				FbxGeometryConverter* SDKGeometryConverter = nullptr;
@@ -91,15 +108,25 @@ namespace UE
 			class FFbxMesh
 			{
 			public:
+				explicit FFbxMesh(FFbxParser& InParser)
+					: Parser(InParser)
+				{}
+
+#if 0
 				static FString GetMeshName(FbxGeometryBase* Mesh);
 				static FString GetMeshUniqueID(FbxGeometryBase* Mesh);
-				static void ExtractSkinnedMeshNodeJoints(FbxScene* SDKScene, FbxMesh* Mesh, UInterchangeMeshNode* MeshNode, TArray<FString>& JSonErrorMessages);
-
-				static void AddAllMeshes(FbxScene* SDKScene, FbxGeometryConverter* SDKGeometryConverter, UInterchangeBaseNodeContainer& NodeContainer, TArray<FString>& JSonErrorMessages, TMap<FString, TSharedPtr<FPayloadContextBase>>& PayloadContexts);
+#endif
+				void ExtractSkinnedMeshNodeJoints(FbxScene* SDKScene, FbxMesh* Mesh, UInterchangeMeshNode* MeshNode);
+				void AddAllMeshes(FbxScene* SDKScene, FbxGeometryConverter* SDKGeometryConverter, UInterchangeBaseNodeContainer& NodeContainer, TMap<FString, TSharedPtr<FPayloadContextBase>>& PayloadContexts);
 			
 			protected:
-				static UInterchangeMeshNode* CreateMeshNode(UInterchangeBaseNodeContainer& NodeContainer, const FString& NodeName, const FString& NodeUniqueID, TArray<FString>& JSonErrorMessages);
+				UInterchangeMeshNode* CreateMeshNode(UInterchangeBaseNodeContainer& NodeContainer, const FString& NodeName, const FString& NodeUniqueID);
+#if 0
 				static FString GetUniqueIDString(const uint64 UniqueID);
+#endif
+
+			private:
+				FFbxParser& Parser;
 			};
 		}//ns Private
 	}//ns Interchange
diff --git a/Engine/Plugins/Runtime/Interchange/Source/Parsers/Fbx/Private/FbxScene.cpp b/Engine/Plugins/Runtime/Interchange/Source/Parsers/Fbx/Private/FbxScene.cpp
--- a/Engine/Plugins/Runtime/Interchange/Source/Parsers/Fbx/Private/FbxScene.cpp
+++ b/Engine/Plugins/Runtime/Interchange/Source/Parsers/Fbx/Private/FbxScene.cpp
@@ -3,169 +3,173 @@
 #include "FbxScene.h"
 
 #include "CoreMinimal.h"
+#include "FbxAPI.h"
 #include "FbxConvert.h"
 #include "FbxHelper.h"
 #include "FbxInclude.h"
 #include "FbxMaterial.h"
 #include "FbxMesh.h"
 #include "InterchangeMeshNode.h"
+#include "InterchangeResultsContainer.h"
 #include "InterchangeSceneNode.h"
 #include "Nodes/InterchangeBaseNodeContainer.h"
 
+#define LOCTEXT_NAMESPACE "InterchangeFbxScene"
+
 namespace UE
 {
 	namespace Interchange
 	{
 		namespace Private
 		{
-			namespace Scene
+
+			void FFbxScene::CreateMeshNodeReference(UInterchangeSceneNode* UnrealSceneNode, FbxNode* FbxSceneNode, FbxNodeAttribute* NodeAttribute, UInterchangeBaseNodeContainer& NodeContainer)
 			{
-				void CreateMeshNodeReference(UInterchangeSceneNode* UnrealSceneNode, FbxNode* FbxSceneNode, FbxNodeAttribute* NodeAttribute, UInterchangeBaseNodeContainer& NodeContainer, TArray<FString>& JSonErrorMessages)
+				UInterchangeMeshNode* MeshNode = nullptr;
+				if (NodeAttribute->GetAttributeType() == FbxNodeAttribute::eMesh)
 				{
-					UInterchangeMeshNode* MeshNode = nullptr;
-					if (NodeAttribute->GetAttributeType() == FbxNodeAttribute::eMesh)
-					{
-						FbxMesh* Mesh = static_cast<FbxMesh*>(NodeAttribute);
-						if (ensure(Mesh))
-						{
-							FString MeshRefString = FFbxMesh::GetMeshUniqueID(Mesh);
-							MeshNode = Cast<UInterchangeMeshNode>(NodeContainer.GetNode(MeshRefString));
-						}
-					}
-					else if (NodeAttribute->GetAttributeType() == FbxNodeAttribute::eShape)
+					FbxMesh* Mesh = static_cast<FbxMesh*>(NodeAttribute);
+					if (ensure(Mesh))
 					{
-						//We do not add a dependency for shape on the scene node since shapes are a MeshNode dependency.
-					}
-
-					if (MeshNode)
-					{
-						UnrealSceneNode->SetCustomMeshDependencyUid(MeshNode->GetUniqueID());
-						MeshNode->SetSceneInstanceUid(UnrealSceneNode->GetUniqueID());
+						FString MeshRefString = FFbxHelper::GetMeshUniqueID(Mesh);
+						MeshNode = Cast<UInterchangeMeshNode>(NodeContainer.GetNode(MeshRefString));
 					}
 				}
-
-				void CreateCameraNodeReference(UInterchangeSceneNode* UnrealSceneNode, FbxNode* FbxSceneNode, FbxNodeAttribute* NodeAttribute, UInterchangeBaseNodeContainer& NodeContainer, TArray<FString>& JSonErrorMessages)
+				else if (NodeAttribute->GetAttributeType() == FbxNodeAttribute::eShape)
 				{
+					//We do not add a dependency for shape on the scene node since shapes are a MeshNode dependency.
 				}
 
-				void CreateLightNodeReference(UInterchangeSceneNode* UnrealSceneNode, FbxNode* FbxSceneNode, FbxNodeAttribute* NodeAttribute, UInterchangeBaseNodeContainer& NodeContainer, TArray<FString>& JSonErrorMessages)
+				if (MeshNode)
 				{
-
+					UnrealSceneNode->SetCustomMeshDependencyUid(MeshNode->GetUniqueID());
+					MeshNode->SetSceneInstanceUid(UnrealSceneNode->GetUniqueID());
 				}
+			}
 
-				UInterchangeSceneNode* AddHierarchyRecursively(FbxNode* Node, FbxScene* SDKScene, UInterchangeBaseNodeContainer& NodeContainer, TArray<FString>& JSonErrorMessages)
-				{
-					FString NodeName = FFbxHelper::GetFbxObjectName(Node);
-					FString NodeUniqueID = FFbxHelper::GetFbxNodeHierarchyName(Node);
+			void FFbxScene::CreateCameraNodeReference(UInterchangeSceneNode* UnrealSceneNode, FbxNode* FbxSceneNode, FbxNodeAttribute* NodeAttribute, UInterchangeBaseNodeContainer& NodeContainer)
+			{
+			}
+
+			void FFbxScene::CreateLightNodeReference(UInterchangeSceneNode* UnrealSceneNode, FbxNode* FbxSceneNode, FbxNodeAttribute* NodeAttribute, UInterchangeBaseNodeContainer& NodeContainer)
+			{
 
-					UInterchangeSceneNode* UnrealNode = FFbxScene::CreateTransformNode(NodeContainer, NodeName, NodeUniqueID, JSonErrorMessages);
-					check(UnrealNode);
+			}
+
+			UInterchangeSceneNode* FFbxScene::AddHierarchyRecursively(FbxNode* Node, FbxScene* SDKScene, UInterchangeBaseNodeContainer& NodeContainer)
+			{
+				FString NodeName = FFbxHelper::GetFbxObjectName(Node);
+				FString NodeUniqueID = FFbxHelper::GetFbxNodeHierarchyName(Node);
+
+				UInterchangeSceneNode* UnrealNode = CreateTransformNode(NodeContainer, NodeName, NodeUniqueID);
+				check(UnrealNode);
 					
-					auto GetConvertedTransform = [](FbxAMatrix& NewFbxMatrix)
-					{
-						FTransform Transform;
-						FbxVector4 NewLocalT = NewFbxMatrix.GetT();
-						FbxVector4 NewLocalS = NewFbxMatrix.GetS();
-						FbxQuaternion NewLocalQ = NewFbxMatrix.GetQ();
-						Transform.SetTranslation(FFbxConvert::ConvertPos(NewLocalT));
-						Transform.SetScale3D(FFbxConvert::ConvertScale(NewLocalS));
-						Transform.SetRotation(FFbxConvert::ConvertRotToQuat(NewLocalQ));
-						return Transform;
-					};
+				auto GetConvertedTransform = [](FbxAMatrix& NewFbxMatrix)
+				{
+					FTransform Transform;
+					FbxVector4 NewLocalT = NewFbxMatrix.GetT();
+					FbxVector4 NewLocalS = NewFbxMatrix.GetS();
+					FbxQuaternion NewLocalQ = NewFbxMatrix.GetQ();
+					Transform.SetTranslation(FFbxConvert::ConvertPos(NewLocalT));
+					Transform.SetScale3D(FFbxConvert::ConvertScale(NewLocalS));
+					Transform.SetRotation(FFbxConvert::ConvertRotToQuat(NewLocalQ));
+					return Transform;
+				};
 
-					{
-						//Set the global node transform
-						FbxAMatrix GloabalFbxMatrix = Node->EvaluateGlobalTransform();
-						FTransform GlobalTransform = GetConvertedTransform(GloabalFbxMatrix);
-						UnrealNode->SetCustomGlobalTransform(GlobalTransform);
-
-						//Set the local node transform
-						FbxAMatrix LocalFbxMatrix = Node->EvaluateLocalTransform();
-						FTransform LocalTransform = GetConvertedTransform(LocalFbxMatrix);
-						UnrealNode->SetCustomLocalTransform(LocalTransform);
-					}
+				{
+					//Set the global node transform
+					FbxAMatrix GloabalFbxMatrix = Node->EvaluateGlobalTransform();
+					FTransform GlobalTransform = GetConvertedTransform(GloabalFbxMatrix);
+					UnrealNode->SetCustomGlobalTransform(GlobalTransform);
+
+					//Set the local node transform
+					FbxAMatrix LocalFbxMatrix = Node->EvaluateLocalTransform();
+					FTransform LocalTransform = GetConvertedTransform(LocalFbxMatrix);
+					UnrealNode->SetCustomLocalTransform(LocalTransform);
+				}
 
-					int32 AttributeCount = Node->GetNodeAttributeCount();
-					for (int32 AttributeIndex = 0; AttributeIndex < AttributeCount; ++AttributeIndex)
+				int32 AttributeCount = Node->GetNodeAttributeCount();
+				for (int32 AttributeIndex = 0; AttributeIndex < AttributeCount; ++AttributeIndex)
+				{
+					FbxNodeAttribute* NodeAttribute = Node->GetNodeAttributeByIndex(AttributeIndex);
+					switch (NodeAttribute->GetAttributeType())
 					{
-						FbxNodeAttribute* NodeAttribute = Node->GetNodeAttributeByIndex(AttributeIndex);
-						switch (NodeAttribute->GetAttributeType())
+						case FbxNodeAttribute::eUnknown:
+						case FbxNodeAttribute::eOpticalReference:
+						case FbxNodeAttribute::eOpticalMarker:
+						case FbxNodeAttribute::eCachedEffect:
+						case FbxNodeAttribute::eNull:
+						case FbxNodeAttribute::eMarker:
+						case FbxNodeAttribute::eCameraStereo:
+						case FbxNodeAttribute::eCameraSwitcher:
+						case FbxNodeAttribute::eNurbs:
+						case FbxNodeAttribute::ePatch:
+						case FbxNodeAttribute::eNurbsCurve:
+						case FbxNodeAttribute::eTrimNurbsSurface:
+						case FbxNodeAttribute::eBoundary:
+						case FbxNodeAttribute::eNurbsSurface:
+						case FbxNodeAttribute::eSubDiv:
+						case FbxNodeAttribute::eLine:
+							//Unsupported attribute
+							break;
+						case FbxNodeAttribute::eShape: //We do not add a dependency for shape on the scene node since shapes are a MeshNode dependency.
+							break;
+						case FbxNodeAttribute::eSkeleton:
+						{
+							//Add the joint specialized type
+							UnrealNode->AddSpecializedType(FSceneNodeStaticData::GetJointSpecializeTypeString());
+							break;
+						}
+						case FbxNodeAttribute::eMesh:
+						{
+							//For Mesh attribute we add the fbx nodes materials
+							FFbxMaterial FbxMaterial(Parser);
+							FbxMaterial.AddAllNodeMaterials(UnrealNode, Node, NodeContainer);
+							CreateMeshNodeReference(UnrealNode, Node, NodeAttribute, NodeContainer);
+							break;
+						}
+						case FbxNodeAttribute::eLODGroup:
 						{
-							case FbxNodeAttribute::eUnknown:
-							case FbxNodeAttribute::eOpticalReference:
-							case FbxNodeAttribute::eOpticalMarker:
-							case FbxNodeAttribute::eCachedEffect:
-							case FbxNodeAttribute::eNull:
-							case FbxNodeAttribute::eMarker:
-							case FbxNodeAttribute::eCameraStereo:
-							case FbxNodeAttribute::eCameraSwitcher:
-							case FbxNodeAttribute::eNurbs:
-							case FbxNodeAttribute::ePatch:
-							case FbxNodeAttribute::eNurbsCurve:
-							case FbxNodeAttribute::eTrimNurbsSurface:
-							case FbxNodeAttribute::eBoundary:
-							case FbxNodeAttribute::eNurbsSurface:
-							case FbxNodeAttribute::eSubDiv:
-							case FbxNodeAttribute::eLine:
-								//Unsupported attribute
-								break;
-							case FbxNodeAttribute::eShape: //We do not add a dependency for shape on the scene node since shapes are a MeshNode dependency.
-								break;
-							case FbxNodeAttribute::eSkeleton:
-							{
-								//Add the joint specialized type
-								UnrealNode->AddSpecializedType(FSceneNodeStaticData::GetJointSpecializeTypeString());
-								break;
-							}
-							case FbxNodeAttribute::eMesh:
-							{
-								//For Mesh attribute we add the fbx nodes materials
-								FFbxMaterial::AddAllNodeMaterials(UnrealNode, Node, NodeContainer, JSonErrorMessages);
-								CreateMeshNodeReference(UnrealNode, Node, NodeAttribute, NodeContainer, JSonErrorMessages);
-								break;
-							}
-							case FbxNodeAttribute::eLODGroup:
-							{
-								UnrealNode->AddSpecializedType(FSceneNodeStaticData::GetLodGroupSpecializeTypeString());
-								break;
-							}
-							case FbxNodeAttribute::eCamera:
-							{
-								//Add the Camera asset
-								CreateCameraNodeReference(UnrealNode, Node, NodeAttribute, NodeContainer, JSonErrorMessages);
-								break;
-							}
-							case FbxNodeAttribute::eLight:
-							{
-								//Add the Light asset
-								CreateLightNodeReference(UnrealNode, Node, NodeAttribute, NodeContainer, JSonErrorMessages);
-								break;
-							}
+							UnrealNode->AddSpecializedType(FSceneNodeStaticData::GetLodGroupSpecializeTypeString());
+							break;
+						}
+						case FbxNodeAttribute::eCamera:
+						{
+							//Add the Camera asset
+							CreateCameraNodeReference(UnrealNode, Node, NodeAttribute, NodeContainer);
+							break;
+						}
+						case FbxNodeAttribute::eLight:
+						{
+							//Add the Light asset
+							CreateLightNodeReference(UnrealNode, Node, NodeAttribute, NodeContainer);
+							break;
 						}
 					}
+				}
 
-					const int32 ChildCount = Node->GetChildCount();
-					for (int32 ChildIndex = 0; ChildIndex < ChildCount; ++ChildIndex)
+				const int32 ChildCount = Node->GetChildCount();
+				for (int32 ChildIndex = 0; ChildIndex < ChildCount; ++ChildIndex)
+				{
+					FbxNode* ChildNode = Node->GetChild(ChildIndex);
+					UInterchangeSceneNode* UnrealChildNode = AddHierarchyRecursively(ChildNode, SDKScene, NodeContainer);
+					if (UnrealChildNode)
 					{
-						FbxNode* ChildNode = Node->GetChild(ChildIndex);
-						UInterchangeSceneNode* UnrealChildNode = AddHierarchyRecursively(ChildNode, SDKScene, NodeContainer, JSonErrorMessages);
-						if (UnrealChildNode)
-						{
-							UnrealChildNode->SetParentUid(UnrealNode->GetUniqueID());
-						}
+						UnrealChildNode->SetParentUid(UnrealNode->GetUniqueID());
 					}
-					return UnrealNode;
 				}
-			} //ns Scene
+				return UnrealNode;
+			}
 
-			UInterchangeSceneNode* FFbxScene::CreateTransformNode(UInterchangeBaseNodeContainer& NodeContainer, const FString& NodeName, const FString& NodeUniqueID, TArray<FString>& JSonErrorMessages)
+			UInterchangeSceneNode* FFbxScene::CreateTransformNode(UInterchangeBaseNodeContainer& NodeContainer, const FString& NodeName, const FString& NodeUniqueID)
 			{
 				FString DisplayLabel(NodeName);
 				FString NodeUid(NodeUniqueID);
 				UInterchangeSceneNode* TransformNode = NewObject<UInterchangeSceneNode>(&NodeContainer, NAME_None);
 				if (!ensure(TransformNode))
 				{
-					JSonErrorMessages.Add(TEXT("{\"Msg\" : {\"Type\" : \"Error\",\n\"Msg\" : \"Cannot allocate a node when importing fbx\"}}"));
+					UInterchangeResultError_Generic* Message = Parser.AddMessage<UInterchangeResultError_Generic>();
+					Message->Text = LOCTEXT("NodeAllocationError", "Unable to allocate a node when importing FBX.");
 					return nullptr;
 				}
 				// Creating a UMaterialInterface
@@ -174,11 +178,13 @@ namespace UE
 				return TransformNode;
 			}
 
-			void FFbxScene::AddHierarchy(FbxScene* SDKScene, UInterchangeBaseNodeContainer& NodeContainer, TArray<FString>& JSonErrorMessages)
+			void FFbxScene::AddHierarchy(FbxScene* SDKScene, UInterchangeBaseNodeContainer& NodeContainer)
 			{
 				 FbxNode* RootNode = SDKScene->GetRootNode();
-				 Scene::AddHierarchyRecursively(RootNode, SDKScene, NodeContainer, JSonErrorMessages);
+				 AddHierarchyRecursively(RootNode, SDKScene, NodeContainer);
 			}
 		} //ns Private
 	} //ns Interchange
 }//ns UE
+
+#undef LOCTEXT_NAMESPACE
diff --git a/Engine/Plugins/Runtime/Interchange/Source/Parsers/Fbx/Private/FbxScene.h b/Engine/Plugins/Runtime/Interchange/Source/Parsers/Fbx/Private/FbxScene.h
--- a/Engine/Plugins/Runtime/Interchange/Source/Parsers/Fbx/Private/FbxScene.h
+++ b/Engine/Plugins/Runtime/Interchange/Source/Parsers/Fbx/Private/FbxScene.h
@@ -15,11 +15,26 @@ namespace UE
 	{
 		namespace Private
 		{
+			class FFbxParser;
+
 			class FFbxScene
 			{
 			public:
-				static void AddHierarchy(FbxScene* SDKScene, UInterchangeBaseNodeContainer& NodeContainer, TArray<FString>& JSonErrorMessages);
-				static UInterchangeSceneNode* CreateTransformNode(UInterchangeBaseNodeContainer& NodeContainer, const FString& NodeName, const FString& NodeUniqueID, TArray<FString>& JSonErrorMessages);
+				explicit FFbxScene(FFbxParser& InParser)
+					: Parser(InParser)
+				{}
+
+				void AddHierarchy(FbxScene* SDKScene, UInterchangeBaseNodeContainer& NodeContainer);
+				UInterchangeSceneNode* CreateTransformNode(UInterchangeBaseNodeContainer& NodeContainer, const FString& NodeName, const FString& NodeUniqueID);
+
+			protected:
+				void CreateMeshNodeReference(UInterchangeSceneNode* UnrealSceneNode, FbxNode* FbxSceneNode, FbxNodeAttribute* NodeAttribute, UInterchangeBaseNodeContainer& NodeContainer);
+				void CreateCameraNodeReference(UInterchangeSceneNode* UnrealSceneNode, FbxNode* FbxSceneNode, FbxNodeAttribute* NodeAttribute, UInterchangeBaseNodeContainer& NodeContainer);
+				void CreateLightNodeReference(UInterchangeSceneNode* UnrealSceneNode, FbxNode* FbxSceneNode, FbxNodeAttribute* NodeAttribute, UInterchangeBaseNodeContainer& NodeContainer);
+				UInterchangeSceneNode* AddHierarchyRecursively(FbxNode* Node, FbxScene* SDKScene, UInterchangeBaseNodeContainer& NodeContainer);
+
+			private:
+				FFbxParser& Parser;
 			};
 		}//ns Private
 	}//ns Interchange
