First pass implementaiton of a message passing system for Interchange. This permits messages, warnings and errors to be output from translators, pipelines and factories, as well as from external proce

diff --git a/Engine/Plugins/Runtime/Interchange/Source/Parsers/Fbx/Private/FbxSkelMesh.cpp b/Engine/Plugins/Runtime/Interchange/Source/Parsers/Fbx/Private/FbxSkelMesh.cpp
new file mode 100644
--- /dev/null
+++ b/Engine/Plugins/Runtime/Interchange/Source/Parsers/Fbx/Private/FbxSkelMesh.cpp
@@ -0,0 +1,931 @@
+// Copyright Epic Games, Inc. All Rights Reserved.
+
+#include "FbxSkelMesh.h"
+
+#include "CoreMinimal.h"
+#include "FbxConvert.h"
+
+#define GeneratedLODNameSuffix "_GeneratedLOD_"
+
+
+namespace UE
+{
+	namespace Interchange
+	{
+		namespace Private
+		{
+			void FFbxSkelMesh::FindSkeletalMeshes(FbxScene* SDKScene, TArray< TArray<FbxNode*> >& outSkelMeshArray, bool bCombineSkeletalMesh, bool bForceFindRigid)
+			{
+				TArray<FbxNode*> SkeletonArray;
+				FbxNode* RootNode = SDKScene->GetRootNode();
+
+				// a) find skeletal meshes
+
+				RecursiveFindFbxSkelMesh(SDKScene, RootNode, outSkelMeshArray, SkeletonArray);
+				// for skeletal mesh, we convert the skeleton system to skeleton
+				// in less we recognize bone mesh as rigid mesh if they are textured
+				for (int32 SkelIndex = 0; SkelIndex < SkeletonArray.Num(); SkelIndex++)
+				{
+					const bool bImportMeshesInBoneHierarchy = true;
+					RecursiveFixSkeleton(SDKScene, SkeletonArray[SkelIndex], outSkelMeshArray[SkelIndex], bImportMeshesInBoneHierarchy);
+				}
+
+
+
+				// b) find rigid mesh
+
+				// If we are attempting to import a skeletal mesh but we have no hierarchy attempt to find a rigid mesh.
+				if (bForceFindRigid || outSkelMeshArray.Num() == 0)
+				{
+					RecursiveFindRigidMesh(SDKScene, RootNode, outSkelMeshArray, SkeletonArray);
+					if (bForceFindRigid)
+					{
+						//Cleanup the rigid mesh, We want to remove any real static mesh from the outSkelMeshArray
+						//Any non skinned mesh that contain no animation should be part of this array.
+						int32 AnimStackCount = SDKScene->GetSrcObjectCount<FbxAnimStack>();
+						TArray<int32> SkeletalMeshArrayToRemove;
+						for (int32 i = 0; i < outSkelMeshArray.Num(); i++)
+						{
+							bool bIsValidSkeletal = false;
+							TArray<FbxNode*>& NodeArray = outSkelMeshArray[i];
+							for (FbxNode* InspectedNode : NodeArray)
+							{
+								FbxMesh* Mesh = InspectedNode->GetMesh();
+
+								FbxLODGroup* LodGroup = InspectedNode->GetLodGroup();
+								if (LodGroup != nullptr)
+								{
+									FbxNode* SkelMeshNode = FindLODGroupNode(InspectedNode, 0, nullptr);
+									if (SkelMeshNode != nullptr)
+									{
+										Mesh = SkelMeshNode->GetMesh();
+									}
+								}
+
+								if (Mesh == nullptr)
+								{
+									continue;
+								}
+								if (Mesh->GetDeformerCount(FbxDeformer::eSkin) > 0)
+								{
+									bIsValidSkeletal = true;
+									break;
+								}
+								//If there is some anim object we count this as a valid skeletal mesh imported as rigid mesh
+								for (int32 AnimStackIndex = 0; AnimStackIndex < AnimStackCount; AnimStackIndex++)
+								{
+									FbxAnimStack* CurAnimStack = SDKScene->GetSrcObject<FbxAnimStack>(AnimStackIndex);
+									// set current anim stack
+
+									SDKScene->SetCurrentAnimationStack(CurAnimStack);
+
+									FbxTimeSpan AnimTimeSpan(FBXSDK_TIME_INFINITE, FBXSDK_TIME_MINUS_INFINITE);
+									InspectedNode->GetAnimationInterval(AnimTimeSpan, CurAnimStack);
+
+									if (AnimTimeSpan.GetDuration() > 0)
+									{
+										bIsValidSkeletal = true;
+										break;
+									}
+								}
+								if (bIsValidSkeletal)
+								{
+									break;
+								}
+							}
+							if (!bIsValidSkeletal)
+							{
+								SkeletalMeshArrayToRemove.Add(i);
+							}
+						}
+						for (int32 i = SkeletalMeshArrayToRemove.Num() - 1; i >= 0; --i)
+						{
+							if (!SkeletalMeshArrayToRemove.IsValidIndex(i) || !outSkelMeshArray.IsValidIndex(SkeletalMeshArrayToRemove[i]))
+								continue;
+							int32 IndexToRemove = SkeletalMeshArrayToRemove[i];
+							outSkelMeshArray[IndexToRemove].Empty();
+							outSkelMeshArray.RemoveAt(IndexToRemove);
+						}
+					}
+				}
+				//Empty the skeleton array
+				SkeletonArray.Empty();
+
+
+				if (bCombineSkeletalMesh)
+				{
+					//Merge all the skeletal mesh arrays into one combine mesh
+					TArray<FbxNode*> CombineNodes;
+					for (TArray<FbxNode*>& Parts : outSkelMeshArray)
+					{
+						CombineNodes.Append(Parts);
+					}
+					outSkelMeshArray.Empty(1);
+					outSkelMeshArray.Add(CombineNodes);
+				}
+			}
+
+			void FFbxSkelMesh::RecursiveGetAllMeshNode(TArray<FbxNode*> & OutAllNode, FbxNode * Node)
+			{
+				if (Node == nullptr)
+				{
+					return;
+				}
+
+				if (Node->GetMesh() != nullptr)
+				{
+					OutAllNode.Add(Node);
+					return;
+				}
+
+				//Look if its a generated LOD
+				FString FbxGeneratedNodeName = UTF8_TO_TCHAR(Node->GetName());
+				if (FbxGeneratedNodeName.Contains(TEXT(GeneratedLODNameSuffix)))
+				{
+					FString SuffixSearch = TEXT(GeneratedLODNameSuffix);
+					int32 SuffixIndex = FbxGeneratedNodeName.Find(SuffixSearch, ESearchCase::CaseSensitive, ESearchDir::FromEnd);
+					SuffixIndex += SuffixSearch.Len();
+					FString LODXNumber = FbxGeneratedNodeName.RightChop(SuffixIndex).Left(1);
+					if (LODXNumber.IsNumeric())
+					{
+						OutAllNode.Add(Node);
+						return;
+					}
+				}
+
+				for (int32 ChildIndex = 0; ChildIndex < Node->GetChildCount(); ++ChildIndex)
+				{
+					RecursiveGetAllMeshNode(OutAllNode, Node->GetChild(ChildIndex));
+				}
+			}
+
+			void FFbxSkelMesh::FindAllLODGroupNode(TArray<FbxNode*> & OutNodeInLod, FbxNode * NodeLodGroup, int32 LodIndex)
+			{
+				check(NodeLodGroup->GetChildCount() >= LodIndex);
+				FbxNode* ChildNode = NodeLodGroup->GetChild(LodIndex);
+				if (ChildNode == nullptr)
+				{
+					return;
+				}
+				RecursiveGetAllMeshNode(OutNodeInLod, ChildNode);
+			}
+
+			bool FFbxSkelMesh::IsUnrealBone(FbxNode * Link)
+			{
+				FbxNodeAttribute* Attr = Link->GetNodeAttribute();
+				if (Attr)
+				{
+					FbxNodeAttribute::EType AttrType = Attr->GetAttributeType();
+					if (AttrType == FbxNodeAttribute::eSkeleton ||
+						AttrType == FbxNodeAttribute::eMesh ||
+						AttrType == FbxNodeAttribute::eNull)
+					{
+						return true;
+					}
+				}
+
+				return false;
+			}
+
+			void FFbxSkelMesh::RecursiveBuildSkeleton(FbxNode * Link, TArray<FbxNode*> & OutSortedLinks)
+			{
+				if (IsUnrealBone(Link))
+				{
+					OutSortedLinks.Add(Link);
+					int32 ChildIndex;
+					for (ChildIndex = 0; ChildIndex < Link->GetChildCount(); ChildIndex++)
+					{
+						RecursiveBuildSkeleton(Link->GetChild(ChildIndex), OutSortedLinks);
+					}
+				}
+			}
+
+			bool FFbxSkelMesh::RetrievePoseFromBindPose(FbxScene * SDKScene, const TArray<FbxNode*> & NodeArray, FbxArray<FbxPose*> & PoseArray)
+			{
+				const int32 PoseCount = SDKScene->GetPoseCount();
+				for (int32 PoseIndex = 0; PoseIndex < PoseCount; PoseIndex++)
+				{
+					FbxPose* CurrentPose = SDKScene->GetPose(PoseIndex);
+
+					// current pose is bind pose, 
+					if (CurrentPose && CurrentPose->IsBindPose())
+					{
+						// IsValidBindPose doesn't work reliably
+						// It checks all the parent chain(regardless root given), and if the parent doesn't have correct bind pose, it fails
+						// It causes more false positive issues than the real issue we have to worry about
+						// If you'd like to try this, set CHECK_VALID_BIND_POSE to 1, and try the error message
+						// when Autodesk fixes this bug, then we might be able to re-open this
+						FString PoseName = CurrentPose->GetName();
+						// all error report status
+						FbxStatus Status;
+
+						// it does not make any difference of checking with different node
+						// it is possible pose 0 -> node array 2, but isValidBindPose function returns true even with node array 0
+						for (auto Current : NodeArray)
+						{
+							FString CurrentName = Current->GetName();
+							NodeList pMissingAncestors, pMissingDeformers, pMissingDeformersAncestors, pWrongMatrices;
+
+							if (CurrentPose->IsValidBindPoseVerbose(Current, pMissingAncestors, pMissingDeformers, pMissingDeformersAncestors, pWrongMatrices, 0.0001, &Status))
+							{
+								PoseArray.Add(CurrentPose);
+								//UE_LOG(LogFbx, Log, TEXT("Valid bind pose for Pose (%s) - %s"), *PoseName, *FString(Current->GetName()));
+								break;
+							}
+							else
+							{
+								// first try to fix up
+								// add missing ancestors
+								for (int i = 0; i < pMissingAncestors.GetCount(); i++)
+								{
+									FbxAMatrix mat = pMissingAncestors.GetAt(i)->EvaluateGlobalTransform(FBXSDK_TIME_ZERO);
+									CurrentPose->Add(pMissingAncestors.GetAt(i), mat);
+								}
+
+								pMissingAncestors.Clear();
+								pMissingDeformers.Clear();
+								pMissingDeformersAncestors.Clear();
+								pWrongMatrices.Clear();
+
+								// check it again
+								if (CurrentPose->IsValidBindPose(Current))
+								{
+									PoseArray.Add(CurrentPose);
+									//UE_LOG(LogFbx, Log, TEXT("Valid bind pose for Pose (%s) - %s"), *PoseName, *FString(Current->GetName()));
+									break;
+								}
+								else
+								{
+									// first try to find parent who is null group and see if you can try test it again
+									FbxNode* ParentNode = Current->GetParent();
+									while (ParentNode)
+									{
+										FbxNodeAttribute* Attr = ParentNode->GetNodeAttribute();
+										if (Attr && Attr->GetAttributeType() == FbxNodeAttribute::eNull)
+										{
+											// found it 
+											break;
+										}
+
+										// find next parent
+										ParentNode = ParentNode->GetParent();
+									}
+
+									if (ParentNode && CurrentPose->IsValidBindPose(ParentNode))
+									{
+										PoseArray.Add(CurrentPose);
+										//UE_LOG(LogFbx, Log, TEXT("Valid bind pose for Pose (%s) - %s"), *PoseName, *FString(Current->GetName()));
+										break;
+									}
+									else
+									{
+										FString ErrorString = Status.GetErrorString();
+										// 											if (!GIsAutomationTesting)
+										// 												UE_LOG(LogFbx, Warning, TEXT("Not valid bind pose for Pose (%s) - Node %s : %s"), *PoseName, *FString(Current->GetName()), *ErrorString);
+									}
+								}
+							}
+						}
+					}
+				}
+
+				return (PoseArray.Size() > 0);
+			}
+
+			FbxNode* FFbxSkelMesh::GetRootSkeleton(FbxScene * SDKScene, FbxNode * Link)
+			{
+				FbxNode* RootBone = Link;
+
+				// get Unreal skeleton root
+				// mesh and dummy are used as bone if they are in the skeleton hierarchy
+				while (RootBone && RootBone->GetParent())
+				{
+					bool bIsBlenderArmatureBone = false;
+					//TODO put back the creator code so we know if its blender
+					//if (FbxCreator == EFbxCreator::Blender)
+					{
+						//Hack to support armature dummy node from blender
+						//Users do not want the null attribute node named armature which is the parent of the real root bone in blender fbx file
+						//This is a hack since if a rigid mesh group root node is named "armature" it will be skip
+						const FString RootBoneParentName(RootBone->GetParent()->GetName());
+						FbxNode* GrandFather = RootBone->GetParent()->GetParent();
+						bIsBlenderArmatureBone = (GrandFather == nullptr || GrandFather == SDKScene->GetRootNode()) && (RootBoneParentName.Compare(TEXT("armature"), ESearchCase::IgnoreCase) == 0);
+					}
+
+					FbxNodeAttribute* Attr = RootBone->GetParent()->GetNodeAttribute();
+					if (Attr &&
+						(Attr->GetAttributeType() == FbxNodeAttribute::eMesh ||
+							(Attr->GetAttributeType() == FbxNodeAttribute::eNull && !bIsBlenderArmatureBone) ||
+							Attr->GetAttributeType() == FbxNodeAttribute::eSkeleton) &&
+						RootBone->GetParent() != SDKScene->GetRootNode())
+					{
+						// in some case, skeletal mesh can be ancestor of bones
+						// this avoids this situation
+						if (Attr->GetAttributeType() == FbxNodeAttribute::eMesh)
+						{
+							FbxMesh* Mesh = (FbxMesh*)Attr;
+							if (Mesh->GetDeformerCount(FbxDeformer::eSkin) > 0)
+							{
+								break;
+							}
+						}
+
+						RootBone = RootBone->GetParent();
+					}
+					else
+					{
+						break;
+					}
+				}
+
+				return RootBone;
+			}
+
+			void FFbxSkelMesh::BuildSkeletonSystem(FbxScene * SDKScene, TArray<FbxCluster*> & ClusterArray, TArray<FbxNode*> & OutSortedLinks)
+			{
+				FbxNode* Link;
+				TArray<FbxNode*> RootLinks;
+				int32 ClusterIndex;
+				for (ClusterIndex = 0; ClusterIndex < ClusterArray.Num(); ClusterIndex++)
+				{
+					Link = ClusterArray[ClusterIndex]->GetLink();
+					if (Link)
+					{
+						Link = GetRootSkeleton(SDKScene, Link);
+						int32 LinkIndex;
+						for (LinkIndex = 0; LinkIndex < RootLinks.Num(); LinkIndex++)
+						{
+							if (Link == RootLinks[LinkIndex])
+							{
+								break;
+							}
+						}
+
+						// this link is a new root, add it
+						if (LinkIndex == RootLinks.Num())
+						{
+							RootLinks.Add(Link);
+						}
+					}
+				}
+
+				for (int32 LinkIndex = 0; LinkIndex < RootLinks.Num(); LinkIndex++)
+				{
+					RecursiveBuildSkeleton(RootLinks[LinkIndex], OutSortedLinks);
+				}
+			}
+
+			bool FFbxSkelMesh::FindSkeletonJoints(FbxScene * SDKScene, TArray<FbxNode*> & NodeArray, TArray<FbxNode*> & SortedLinks, FbxArray<FbxAMatrix> & LocalsPerLink)
+			{
+				bool bIsARigidMesh = false;
+				FbxNode* Link = NULL;
+				FbxArray<FbxPose*> PoseArray;
+				TArray<FbxCluster*> ClusterArray;
+
+				if (NodeArray[0]->GetMesh()->GetDeformerCount(FbxDeformer::eSkin) == 0)
+				{
+					bIsARigidMesh = true;
+					Link = NodeArray[0];
+					RecursiveBuildSkeleton(GetRootSkeleton(SDKScene, Link), SortedLinks);
+				}
+				else
+				{
+					// get bindpose and clusters from FBX skeleton
+
+					// let's put the elements to their bind pose! (and we restore them after
+					// we have built the ClusterInformation.
+					int32 Default_NbPoses = SDKScene->GetFbxManager()->GetBindPoseCount(SDKScene);
+					// If there are no BindPoses, the following will generate them.
+					SDKScene->GetFbxManager()->CreateMissingBindPoses(SDKScene);
+
+					//if we created missing bind poses, update the number of bind poses
+					int32 NbPoses = SDKScene->GetFbxManager()->GetBindPoseCount(SDKScene);
+
+					if (NbPoses != Default_NbPoses)
+					{
+						//TODO output an error message to let the user know there is a missing bind pose when translating this fbx skeletalmesh
+						//Missing bind pose
+					}
+
+					//
+					// create the bones / skinning
+					//
+
+					for (int32 i = 0; i < NodeArray.Num(); i++)
+					{
+						FbxMesh* FbxMesh = NodeArray[i]->GetMesh();
+						const int32 SkinDeformerCount = FbxMesh->GetDeformerCount(FbxDeformer::eSkin);
+						for (int32 DeformerIndex = 0; DeformerIndex < SkinDeformerCount; DeformerIndex++)
+						{
+							FbxSkin* Skin = (FbxSkin*)FbxMesh->GetDeformer(DeformerIndex, FbxDeformer::eSkin);
+							for (int32 ClusterIndex = 0; ClusterIndex < Skin->GetClusterCount(); ClusterIndex++)
+							{
+								ClusterArray.Add(Skin->GetCluster(ClusterIndex));
+							}
+						}
+					}
+
+					if (ClusterArray.Num() == 0)
+					{
+						//TODO warn user we found a mesh with a deformer that is not bind to any skin cluster
+						return false;
+					}
+
+					// get bind pose
+					if (RetrievePoseFromBindPose(SDKScene, NodeArray, PoseArray) == false)
+					{
+						// 							if (!GIsAutomationTesting)
+						// 								UE_LOG(LogFbx, Warning, TEXT("Getting valid bind pose failed. Try to recreate bind pose"));
+						// if failed, delete bind pose, and retry.
+						const int32 PoseCount = SDKScene->GetPoseCount();
+						for (int32 PoseIndex = PoseCount - 1; PoseIndex >= 0; --PoseIndex)
+						{
+							FbxPose* CurrentPose = SDKScene->GetPose(PoseIndex);
+
+							// current pose is bind pose, 
+							if (CurrentPose && CurrentPose->IsBindPose())
+							{
+								SDKScene->RemovePose(PoseIndex);
+								CurrentPose->Destroy();
+							}
+						}
+
+						SDKScene->GetFbxManager()->CreateMissingBindPoses(SDKScene);
+						if (RetrievePoseFromBindPose(SDKScene, NodeArray, PoseArray) == false)
+						{
+							// 								if (!GIsAutomationTesting)
+							// 									UE_LOG(LogFbx, Warning, TEXT("Recreating bind pose failed."));
+						}
+						else
+						{
+							// 								if (!GIsAutomationTesting)
+							// 									UE_LOG(LogFbx, Warning, TEXT("Recreating bind pose succeeded."));
+						}
+					}
+
+					// recurse through skeleton and build ordered table
+					BuildSkeletonSystem(SDKScene, ClusterArray, SortedLinks);
+				}
+
+				// error check
+				// if no bond is found
+				if (SortedLinks.Num() == 0)
+				{
+					//						{
+					//							AddTokenizedErrorMessage(FTokenizedMessage::Create(EMessageSeverity::Warning, FText::Format(LOCTEXT("FbxSkeletaLMeshimport_NoBone", "'{0}' has no bones"), FText::FromString(NodeArray[0]->GetName()))), FFbxErrors::SkeletalMesh_NoBoneFound);
+					//						}
+					return false;
+				}
+
+				// if no bind pose is found
+				if (PoseArray.GetCount() == 0)
+				{
+					//TODO error to the user
+					//AddTokenizedErrorMessage(FTokenizedMessage::Create(EMessageSeverity::Warning, LOCTEXT("FbxSkeletaLMeshimport_MissingBindPose", "Could not find the bind pose.  It will use time 0 as bind pose.")), FFbxErrors::SkeletalMesh_InvalidBindPose);
+				}
+
+				int32 LinkIndex;
+
+				// Check for duplicate bone names and issue a warning if found
+				for (LinkIndex = 0; LinkIndex < SortedLinks.Num(); ++LinkIndex)
+				{
+					Link = SortedLinks[LinkIndex];
+
+					for (int32 AltLinkIndex = LinkIndex + 1; AltLinkIndex < SortedLinks.Num(); ++AltLinkIndex)
+					{
+						FbxNode* AltLink = SortedLinks[AltLinkIndex];
+
+						if (FCStringAnsi::Strcmp(Link->GetName(), AltLink->GetName()) == 0)
+						{
+							//								FString RawBoneName = UTF8_TO_TCHAR(Link->GetName());
+							//TODO notify the user that there is a duplicate bone name
+							// 									AddTokenizedErrorMessage(FTokenizedMessage::Create(EMessageSeverity::Error, FText::Format(LOCTEXT("Error_DuplicateBoneName", "Error, Could not import {0}.\nDuplicate bone name found ('{1}'). Each bone must have a unique name."),
+							// 															 FText::FromString(NodeArray[0]->GetName()), FText::FromString(RawBoneName))), FFbxErrors::SkeletalMesh_DuplicateBones);
+							return false;
+						}
+					}
+				}
+
+				FbxArray<FbxAMatrix> GlobalsPerLink;
+				GlobalsPerLink.Grow(SortedLinks.Num());
+				GlobalsPerLink[0].SetIdentity();
+				LocalsPerLink.Grow(SortedLinks.Num());
+				LocalsPerLink[0].SetIdentity();
+
+				bool GlobalLinkFoundFlag;
+
+				bool bAnyLinksNotInBindPose = false;
+				FString LinksWithoutBindPoses;
+
+				int32 RootIdx = INDEX_NONE;
+
+				for (LinkIndex = 0; LinkIndex < SortedLinks.Num(); LinkIndex++)
+				{
+					// Add a bone for each FBX Link
+					Link = SortedLinks[LinkIndex];
+					int32 ParentIndex = INDEX_NONE; // base value for root if no parent found
+
+					if (LinkIndex == 0)
+					{
+						RootIdx = LinkIndex;
+					}
+					else
+					{
+						const FbxNode* LinkParent = Link->GetParent();
+						// get the link parent index.
+						for (int32 ll = 0; ll < LinkIndex; ++ll) // <LinkIndex because parent is guaranteed to be before child in sortedLink
+						{
+							FbxNode* Otherlink = SortedLinks[ll];
+							if (Otherlink == LinkParent)
+							{
+								ParentIndex = ll;
+								break;
+							}
+						}
+
+						if (ParentIndex == INDEX_NONE)//We found another root inside the hierarchy, this is not supported
+						{
+							//AddTokenizedErrorMessage(FTokenizedMessage::Create(EMessageSeverity::Error, LOCTEXT("MultipleRootsFound", "Multiple roots are found in the bone hierarchy. We only support single root bone.")), FFbxErrors::SkeletalMesh_MultipleRoots);
+							return false;
+						}
+					}
+
+					GlobalLinkFoundFlag = false;
+					if (!bIsARigidMesh) //skeletal mesh
+					{
+						// there are some links, they have no cluster, but in bindpose
+						if (PoseArray.GetCount())
+						{
+							for (int32 PoseIndex = 0; PoseIndex < PoseArray.GetCount(); PoseIndex++)
+							{
+								int32 PoseLinkIndex = PoseArray[PoseIndex]->Find(Link);
+								if (PoseLinkIndex >= 0)
+								{
+									FbxMatrix NoneAffineMatrix = PoseArray[PoseIndex]->GetMatrix(PoseLinkIndex);
+									FbxAMatrix Matrix = *(FbxAMatrix*)(double*)&NoneAffineMatrix;
+									GlobalsPerLink[LinkIndex] = Matrix;
+									GlobalLinkFoundFlag = true;
+									break;
+								}
+							}
+						}
+
+						if (!GlobalLinkFoundFlag)
+						{
+							// since now we set use time 0 as ref pose this won't unlikely happen
+							// but leaving it just in case it still has case where it's missing partial bind pose
+							bAnyLinksNotInBindPose = true;
+							LinksWithoutBindPoses += UTF8_TO_TCHAR(Link->GetName());
+							LinksWithoutBindPoses += TEXT("  \n");
+
+							for (int32 ClusterIndex = 0; ClusterIndex < ClusterArray.Num(); ClusterIndex++)
+							{
+								FbxCluster* Cluster = ClusterArray[ClusterIndex];
+								if (Link == Cluster->GetLink())
+								{
+									Cluster->GetTransformLinkMatrix(GlobalsPerLink[LinkIndex]);
+									GlobalLinkFoundFlag = true;
+									break;
+								}
+							}
+						}
+					}
+
+					if (!GlobalLinkFoundFlag)
+					{
+						GlobalsPerLink[LinkIndex] = Link->EvaluateGlobalTransform();
+					}
+
+
+					if (LinkIndex)
+					{
+						LocalsPerLink[LinkIndex] = GlobalsPerLink[ParentIndex].Inverse() * GlobalsPerLink[LinkIndex];
+					}
+					else	// skeleton root
+					{
+						LocalsPerLink[LinkIndex] = GlobalsPerLink[LinkIndex];
+					}
+
+					//Check for nan and for zero scale
+					{
+						bool bFoundNan = false;
+						bool bFoundZeroScale = false;
+						for (int32 i = 0; i < 4; ++i)
+						{
+							if (i < 3)
+							{
+								if (FMath::IsNaN(LocalsPerLink[LinkIndex].GetT()[i]) || FMath::IsNaN(LocalsPerLink[LinkIndex].GetS()[i]))
+								{
+									bFoundNan = true;
+								}
+								if (FMath::IsNearlyZero(LocalsPerLink[LinkIndex].GetS()[i]))
+								{
+									bFoundZeroScale = true;
+								}
+							}
+							if (FMath::IsNaN(LocalsPerLink[LinkIndex].GetQ()[i]))
+							{
+								bFoundNan = true;
+							}
+						}
+
+						//TODO add a user error message
+						ensure(!bFoundNan);
+						ensure(!bFoundZeroScale);
+					}
+				}
+
+
+
+				if (bAnyLinksNotInBindPose)
+				{
+					//TODO error message
+					//
+					//AddTokenizedErrorMessage(FTokenizedMessage::Create(EMessageSeverity::Warning, FText::Format(LOCTEXT("FbxSkeletaLMeshimport_BonesAreMissingFromBindPose", "The following bones are missing from the bind pose:\n{0}\nThis can happen for bones that are not vert weighted. If they are not in the correct orientation after importing,\nplease set the \"Use T0 as ref pose\" option or add them to the bind pose and reimport the skeletal mesh."), FText::FromString(LinksWithoutBindPoses))), FFbxErrors::SkeletalMesh_BonesAreMissingFromBindPose);
+				}
+				return true;
+			}
+
+			FbxNode* FFbxSkelMesh::RecursiveGetFirstMeshNode(FbxNode* Node, FbxNode* NodeToFind)
+			{
+				if (Node == nullptr)
+				{
+					return nullptr;
+				}
+				if (Node->GetMesh() != nullptr)
+					return Node;
+				for (int32 ChildIndex = 0; ChildIndex < Node->GetChildCount(); ++ChildIndex)
+				{
+					FbxNode* MeshNode = RecursiveGetFirstMeshNode(Node->GetChild(ChildIndex), NodeToFind);
+					if (NodeToFind == nullptr)
+					{
+						if (MeshNode != nullptr)
+						{
+							return MeshNode;
+						}
+					}
+					else if (MeshNode == NodeToFind)
+					{
+						return MeshNode;
+					}
+				}
+				return nullptr;
+			}
+
+			FbxNode* FFbxSkelMesh::FindLODGroupNode(FbxNode* NodeLodGroup, int32 LodIndex, FbxNode* NodeToFind)
+			{
+				check(NodeLodGroup->GetChildCount() >= LodIndex);
+				FbxNode* ChildNode = NodeLodGroup->GetChild(LodIndex);
+				if (ChildNode == nullptr)
+				{
+					return nullptr;
+				}
+				return RecursiveGetFirstMeshNode(ChildNode, NodeToFind);
+			}
+
+			void FFbxSkelMesh::RecursiveFindFbxSkelMesh(FbxScene * SDKScene, FbxNode * Node, TArray< TArray<FbxNode*> > & outSkelMeshArray, TArray<FbxNode*> & SkeletonArray)
+			{
+				FbxNode* SkelMeshNode = nullptr;
+				FbxNode* NodeToAdd = Node;
+
+				if (Node->GetMesh() && Node->GetMesh()->GetDeformerCount(FbxDeformer::eSkin) > 0)
+				{
+					SkelMeshNode = Node;
+				}
+				else if (Node->GetNodeAttribute() && Node->GetNodeAttribute()->GetAttributeType() == FbxNodeAttribute::eLODGroup)
+				{
+					// for LODgroup, add the LODgroup to OutSkelMeshArray according to the skeleton that the first child bind to
+					SkelMeshNode = FindLODGroupNode(Node, 0, nullptr);
+					// check if the first child is skeletal mesh
+					if (SkelMeshNode != nullptr && !(SkelMeshNode->GetMesh() && SkelMeshNode->GetMesh()->GetDeformerCount(FbxDeformer::eSkin) > 0))
+					{
+						SkelMeshNode = nullptr;
+					}
+					// else NodeToAdd = Node;
+				}
+
+				if (SkelMeshNode)
+				{
+					// find root skeleton
+
+					check(SkelMeshNode->GetMesh() != nullptr);
+					const int32 fbxDeformerCount = SkelMeshNode->GetMesh()->GetDeformerCount();
+					FbxSkin* Deformer = static_cast<FbxSkin*>(SkelMeshNode->GetMesh()->GetDeformer(0, FbxDeformer::eSkin));
+
+					if (Deformer != NULL)
+					{
+						int32 ClusterCount = Deformer->GetClusterCount();
+						bool bFoundCorrectLink = false;
+						for (int32 ClusterId = 0; ClusterId < ClusterCount; ++ClusterId)
+						{
+							FbxNode* RootBoneLink = Deformer->GetCluster(ClusterId)->GetLink(); //Get the bone influences by this first cluster
+							RootBoneLink = GetRootSkeleton(SDKScene, RootBoneLink); // Get the skeleton root itself
+
+							if (RootBoneLink)
+							{
+								bool bAddedToExistingSkeleton = false;
+								for (int32 SkeletonIndex = 0; SkeletonIndex < SkeletonArray.Num(); ++SkeletonIndex)
+								{
+									if (RootBoneLink == SkeletonArray[SkeletonIndex])
+									{
+										// append to existed outSkelMeshArray element
+										TArray<FbxNode*>& TempArray = outSkelMeshArray[SkeletonIndex];
+										TempArray.Add(NodeToAdd);
+										bAddedToExistingSkeleton = true;
+										break;
+									}
+								}
+
+								// if there is no outSkelMeshArray element that is bind to this skeleton
+								// create new element for outSkelMeshArray
+								if (!bAddedToExistingSkeleton)
+								{
+									TArray<FbxNode*>& TempArray = outSkelMeshArray.AddDefaulted_GetRef();
+									TempArray.Add(NodeToAdd);
+									SkeletonArray.Add(RootBoneLink);
+								}
+
+								bFoundCorrectLink = true;
+								break;
+							}
+						}
+
+						// we didn't find the correct link
+						if (!bFoundCorrectLink)
+						{
+							// 								AddTokenizedErrorMessage(
+							// 									FTokenizedMessage::Create(
+							// 									EMessageSeverity::Warning,
+							// 									FText::Format(LOCTEXT("FBX_NoWeightsOnDeformer", "Ignoring mesh {0} because it but no weights."), FText::FromString(UTF8_TO_TCHAR(SkelMeshNode->GetName())))
+							// 								),
+							// 									FFbxErrors::SkeletalMesh_NoWeightsOnDeformer
+							// 								);
+						}
+					}
+				}
+
+				//Skeletalmesh node can have child so let's always iterate trough child
+				{
+					TArray<FbxNode*> ChildScaled;
+					//Sort the node to have the one with no scaling first so we have more chance
+					//to have a root skeletal mesh with no scale. Because scene import do not support
+					//root skeletal mesh containing scale
+					for (int32 ChildIndex = 0; ChildIndex < Node->GetChildCount(); ++ChildIndex)
+					{
+						FbxNode* ChildNode = Node->GetChild(ChildIndex);
+
+						if (!Node->GetNodeAttribute() || Node->GetNodeAttribute()->GetAttributeType() != FbxNodeAttribute::eLODGroup)
+						{
+							FbxVector4 NoScale(1.0, 1.0, 1.0);
+
+							if (ChildNode->EvaluateLocalScaling() == NoScale)
+							{
+								RecursiveFindFbxSkelMesh(SDKScene, ChildNode, outSkelMeshArray, SkeletonArray);
+							}
+							else
+							{
+								ChildScaled.Add(ChildNode);
+							}
+						}
+					}
+
+					for (FbxNode* ChildNode : ChildScaled)
+					{
+						RecursiveFindFbxSkelMesh(SDKScene, ChildNode, outSkelMeshArray, SkeletonArray);
+					}
+				}
+			}
+
+			void FFbxSkelMesh::RecursiveFindRigidMesh(FbxScene * SDKScene, FbxNode * Node, TArray< TArray<FbxNode*> > & outSkelMeshArray, TArray<FbxNode*> & SkeletonArray)
+			{
+				bool bRigidNodeFound = false;
+				FbxNode* RigidMeshNode = nullptr;
+
+				if (Node->GetMesh())
+				{
+					// ignore skeletal mesh
+					if (Node->GetMesh()->GetDeformerCount(FbxDeformer::eSkin) == 0)
+					{
+						RigidMeshNode = Node;
+						bRigidNodeFound = true;
+					}
+				}
+				else if (Node->GetNodeAttribute() && Node->GetNodeAttribute()->GetAttributeType() == FbxNodeAttribute::eLODGroup)
+				{
+					// for LODgroup, add the LODgroup to OutSkelMeshArray according to the skeleton that the first child bind to
+					FbxNode* FirstLOD = FindLODGroupNode(Node, 0, nullptr);
+					// check if the first child is skeletal mesh
+					if (FirstLOD != nullptr && FirstLOD->GetMesh())
+					{
+						if (FirstLOD->GetMesh()->GetDeformerCount(FbxDeformer::eSkin) == 0)
+						{
+							bRigidNodeFound = true;
+						}
+					}
+
+					if (bRigidNodeFound)
+					{
+						RigidMeshNode = Node;
+					}
+				}
+
+				if (bRigidNodeFound)
+				{
+					// find root skeleton
+					FbxNode* Link = GetRootSkeleton(SDKScene, RigidMeshNode);
+
+					int32 i;
+					for (i = 0; i < SkeletonArray.Num(); i++)
+					{
+						if (Link == SkeletonArray[i])
+						{
+							// append to existed outSkelMeshArray element
+							TArray<FbxNode*>& TempArray = outSkelMeshArray[i];
+							TempArray.Add(RigidMeshNode);
+							break;
+						}
+					}
+
+					// if there is no outSkelMeshArray element that is bind to this skeleton
+					// create new element for outSkelMeshArray
+					if (i == SkeletonArray.Num())
+					{
+						TArray<FbxNode*>& TempArray = outSkelMeshArray.AddDefaulted_GetRef();
+						TempArray.Add(RigidMeshNode);
+						SkeletonArray.Add(Link);
+					}
+				}
+
+				// for LODGroup, we will not deep in.
+				if (!(Node->GetNodeAttribute() && Node->GetNodeAttribute()->GetAttributeType() == FbxNodeAttribute::eLODGroup))
+				{
+					int32 ChildIndex;
+					for (ChildIndex = 0; ChildIndex < Node->GetChildCount(); ++ChildIndex)
+					{
+						RecursiveFindRigidMesh(SDKScene, Node->GetChild(ChildIndex), outSkelMeshArray, SkeletonArray);
+					}
+				}
+			}
+
+			void FFbxSkelMesh::RecursiveFixSkeleton(FbxScene * SDKScene, FbxNode * Node, TArray<FbxNode*> & SkelMeshes, bool bImportNestedMeshes)
+			{
+				FbxNodeAttribute* Attr = Node->GetNodeAttribute();
+				bool NodeIsLodGroup = (Attr && (Attr->GetAttributeType() == FbxNodeAttribute::eLODGroup));
+				if (!NodeIsLodGroup)
+				{
+					for (int32 i = 0; i < Node->GetChildCount(); i++)
+					{
+						RecursiveFixSkeleton(SDKScene, Node->GetChild(i), SkelMeshes, bImportNestedMeshes);
+					}
+				}
+
+				if (Attr && (Attr->GetAttributeType() == FbxNodeAttribute::eMesh || Attr->GetAttributeType() == FbxNodeAttribute::eNull))
+				{
+					if (bImportNestedMeshes && Attr->GetAttributeType() == FbxNodeAttribute::eMesh)
+					{
+						// for leaf mesh, keep them as mesh
+						int32 ChildCount = Node->GetChildCount();
+						int32 ChildIndex;
+						for (ChildIndex = 0; ChildIndex < ChildCount; ChildIndex++)
+						{
+							FbxNode* Child = Node->GetChild(ChildIndex);
+							if (Child->GetMesh() == NULL)
+							{
+								break;
+							}
+						}
+
+						if (ChildIndex != ChildCount)
+						{
+							// Remove from the mesh list it is no longer a mesh
+							SkelMeshes.Remove(Node);
+
+							//replace with skeleton
+							FbxSkeleton* lSkeleton = FbxSkeleton::Create(SDKScene->GetFbxManager(), "");
+							Node->SetNodeAttribute(lSkeleton);
+							lSkeleton->SetSkeletonType(FbxSkeleton::eLimbNode);
+						}
+						else // this mesh may be not in skeleton mesh list. If not, add it.
+						{
+							if (!SkelMeshes.Contains(Node))
+							{
+								SkelMeshes.Add(Node);
+							}
+						}
+					}
+					else
+					{
+						// Remove from the mesh list it is no longer a mesh
+						SkelMeshes.Remove(Node);
+
+						//replace with skeleton
+						FbxSkeleton* lSkeleton = FbxSkeleton::Create(SDKScene->GetFbxManager(), "");
+						Node->SetNodeAttribute(lSkeleton);
+						lSkeleton->SetSkeletonType(FbxSkeleton::eLimbNode);
+					}
+				}
+			}
+
+		}//ns Private
+	}//ns Interchange
+}//ns UE
diff --git a/Engine/Plugins/Runtime/Interchange/Source/Parsers/Fbx/Private/FbxSkelMesh.h b/Engine/Plugins/Runtime/Interchange/Source/Parsers/Fbx/Private/FbxSkelMesh.h
new file mode 100644
--- /dev/null
+++ b/Engine/Plugins/Runtime/Interchange/Source/Parsers/Fbx/Private/FbxSkelMesh.h
@@ -0,0 +1,55 @@
+// Copyright Epic Games, Inc. All Rights Reserved.
+
+#pragma once
+
+#include "CoreMinimal.h"
+#include "FbxAPI.h"
+#include "FbxHelper.h"
+
+namespace UE
+{
+	namespace Interchange
+	{
+		namespace Private
+		{
+			struct FFbxSkelMesh
+			{
+			public:
+				explicit FFbxSkelMesh(FFbxParser& InParser)
+					: Parser(InParser)
+				{}
+
+				void FindSkeletalMeshes(FbxScene* SDKScene, TArray<TArray<FbxNode*>>& outSkelMeshArray, bool bCombineSkeletalMesh, bool bForceFindRigid);
+				void FindAllLODGroupNode(TArray<FbxNode*>& OutNodeInLod, FbxNode* NodeLodGroup, int32 LodIndex);
+				bool FindSkeletonJoints(FbxScene* SDKScene, TArray<FbxNode*>& NodeArray, TArray<FbxNode*>& SortedLinks, FbxArray<FbxAMatrix>& LocalsPerLink);
+				FbxNode* FindLODGroupNode(FbxNode* NodeLodGroup, int32 LodIndex, FbxNode* NodeToFind);
+
+				/**
+				 * Add messages to the message log
+				 */
+				template <typename T>
+				T* AddMessage(FbxGeometryBase* FbxNode) const
+				{
+					T* Item = Parser.AddMessage<T>();
+					Item->MeshName = FFbxHelper::GetMeshName(FbxNode);
+					Item->InterchangeKey = FFbxHelper::GetMeshUniqueID(FbxNode);
+					return Item;
+				}
+
+			private:
+				void RecursiveGetAllMeshNode(TArray<FbxNode*>& OutAllNode, FbxNode* Node);
+				bool IsUnrealBone(FbxNode* Link);
+				void RecursiveBuildSkeleton(FbxNode* Link, TArray<FbxNode*>& OutSortedLinks);
+				bool RetrievePoseFromBindPose(FbxScene* SDKScene, const TArray<FbxNode*>& NodeArray, FbxArray<FbxPose*>& PoseArray);
+				FbxNode* GetRootSkeleton(FbxScene* SDKScene, FbxNode* Link);
+				void BuildSkeletonSystem(FbxScene* SDKScene, TArray<FbxCluster*>& ClusterArray, TArray<FbxNode*>& OutSortedLinks);
+				FbxNode* RecursiveGetFirstMeshNode(FbxNode* Node, FbxNode* NodeToFind);
+				void RecursiveFindFbxSkelMesh(FbxScene* SDKScene, FbxNode* Node, TArray< TArray<FbxNode*> >& outSkelMeshArray, TArray<FbxNode*>& SkeletonArray);
+				void RecursiveFindRigidMesh(FbxScene* SDKScene, FbxNode* Node, TArray< TArray<FbxNode*> >& outSkelMeshArray, TArray<FbxNode*>& SkeletonArray);
+				void RecursiveFixSkeleton(FbxScene* SDKScene, FbxNode* Node, TArray<FbxNode*>& SkelMeshes, bool bImportNestedMeshes);
+
+				FFbxParser& Parser;
+			};
+		}//ns Private
+	}//ns Interchange
+}//ns UE
diff --git a/Engine/Plugins/Runtime/Interchange/Source/Parsers/Fbx/Private/InterchangeFbxParser.cpp b/Engine/Plugins/Runtime/Interchange/Source/Parsers/Fbx/Private/InterchangeFbxParser.cpp
--- a/Engine/Plugins/Runtime/Interchange/Source/Parsers/Fbx/Private/InterchangeFbxParser.cpp
+++ b/Engine/Plugins/Runtime/Interchange/Source/Parsers/Fbx/Private/InterchangeFbxParser.cpp
@@ -10,60 +10,87 @@
 #include "UObject/UObjectGlobals.h"
 #include "UObject/Package.h"
 
+#define LOCTEXT_NAMESPACE "InterchangeFbxParser"
+
 namespace UE
 {
 	namespace Interchange
 	{
 		FInterchangeFbxParser::FInterchangeFbxParser()
 		{
-			FbxParserPrivate = MakeUnique<Private::FFbxParser>();
+			ResultsContainer = NewObject<UInterchangeResultsContainer>(GetTransientPackage());
+			ResultsContainer->AddToRoot();
+			FbxParserPrivate = MakeUnique<Private::FFbxParser>(ResultsContainer);
 		}
+
 		FInterchangeFbxParser::~FInterchangeFbxParser()
 		{
+			ResultsContainer->RemoveFromRoot();
 			FbxParserPrivate = nullptr;
 		}
 
 		void FInterchangeFbxParser::LoadFbxFile(const FString& Filename, const FString& ResultFolder)
 		{
 			check(FbxParserPrivate.IsValid());
+			SourceFilename = Filename;
+			ResultsContainer->Empty();
 
-			if (!FbxParserPrivate->LoadFbxFile(Filename, JsonLoadMessages))
+			if (!FbxParserPrivate->LoadFbxFile(Filename))
 			{
-				JsonLoadMessages.Add(TEXT("{\"Msg\" : {\"Type\" : \"Error\",\n\"Msg\" : \"Cannot load the fbx file.\"}}"));
+				UInterchangeResultError_Generic* Error = AddMessage<UInterchangeResultError_Generic>();
+				Error->SourceAssetName = SourceFilename;
+				Error->Text = LOCTEXT("CantLoadFbxFile", "Cannot load the FBX file.");
 				return;
 			}
 
-
 			ResultFilepath = ResultFolder + TEXT("/SceneDescription.itc");
 			//Since we are not in main thread we cannot use TStrongPtr, so we will add the object to the root and remove it when we are done
 			UInterchangeBaseNodeContainer* Container = NewObject<UInterchangeBaseNodeContainer>(GetTransientPackage(), NAME_None);
 			if (!ensure(Container != nullptr))
 			{
-				JsonLoadMessages.Add(TEXT("{\"Msg\" : {\"Type\" : \"Error\",\n\"Msg\" : \"Cannot allocate base node container to add fbx scene data\"}}"));
+				UInterchangeResultError_Generic* Error = AddMessage<UInterchangeResultError_Generic>();
+				Error->SourceAssetName = SourceFilename;
+				Error->Text = LOCTEXT("CantAllocate", "Cannot allocate base node container to add FBX scene data.");
 				return;
 			}
+
 			Container->AddToRoot();
-			FbxParserPrivate->FillContainerWithFbxScene(*Container, JsonLoadMessages);
+			FbxParserPrivate->FillContainerWithFbxScene(*Container);
 			Container->SaveToFile(ResultFilepath);
 			Container->RemoveFromRoot();
-			JsonLoadMessages.Add(TEXT("{\"Msg\" : {\"Type\" : \"Log\",\n\"Msg\" : \"This is a success!\"}}"));
 		}
 
 		void FInterchangeFbxParser::FetchPayload(const FString& PayloadKey, const FString& ResultFolder)
 		{
 			check(FbxParserPrivate.IsValid());
+			ResultsContainer->Empty();
 
 			FString& PayloadFilepath = ResultPayloads.FindOrAdd(PayloadKey);
 			PayloadFilepath = ResultFolder + TEXT("/") + PayloadKey + TEXT(".payload");
+
 			//Copy the map filename key because we are multithreaded and the TMap can be reallocated
 			FString PayloadFilepathCopy = PayloadFilepath;
-			if (!FbxParserPrivate->FetchPayloadData(PayloadKey, PayloadFilepathCopy, JsonLoadMessages))
+			if (!FbxParserPrivate->FetchPayloadData(PayloadKey, PayloadFilepathCopy))
 			{
-				JsonLoadMessages.Add(TEXT("{\"Msg\" : {\"Type\" : \"Error\",\n\"Msg\" : \"Cannot fetch fbx payload data.\"}}"));
+				UInterchangeResultError_Generic* Error = AddMessage<UInterchangeResultError_Generic>();
+				Error->SourceAssetName = SourceFilename;
+				Error->Text = LOCTEXT("CantFetchPayload", "Cannot fetch FBX payload data.");
 				return;
 			}
-			
-			JsonLoadMessages.Add(TEXT("{\"Msg\" : {\"Type\" : \"Log\",\n\"Msg\" : \"Fetch fbx Payload success!\"}}"));
 		}
+
+		TArray<FString> FInterchangeFbxParser::GetJsonLoadMessages() const
+		{
+			TArray<FString> JsonResults;
+			for (UInterchangeResult* Result : ResultsContainer->GetResults())
+			{
+				JsonResults.Add(Result->ToJson());
+			}
+
+			return JsonResults;
+		}
+
 	}//ns Interchange
 }//ns UE
+
+#undef LOCTEXT_NAMESPACE
diff --git a/Engine/Plugins/Runtime/Interchange/Source/Parsers/Fbx/Public/InterchangeFbxParser.h b/Engine/Plugins/Runtime/Interchange/Source/Parsers/Fbx/Public/InterchangeFbxParser.h
--- a/Engine/Plugins/Runtime/Interchange/Source/Parsers/Fbx/Public/InterchangeFbxParser.h
+++ b/Engine/Plugins/Runtime/Interchange/Source/Parsers/Fbx/Public/InterchangeFbxParser.h
@@ -4,6 +4,7 @@
 
 #include "CoreMinimal.h"
 #include "Nodes/InterchangeBaseNodeContainer.h"
+#include "InterchangeResultsContainer.h"
 #include "UObject/StrongObjectPtr.h"
 
 namespace UE
@@ -43,11 +44,23 @@ namespace UE
 				}
 				return FString();
 			}
-			TArray<FString> GetJsonLoadMessages() const { return JsonLoadMessages; }
+
+			/**
+			 * Transform the results container into an array of Json strings
+			 */
+			TArray<FString> GetJsonLoadMessages() const;
+
+			template <typename T>
+			T* AddMessage()
+			{
+				return ResultsContainer->Add<T>();
+			}
+
 
 		private:
+			TObjectPtr<UInterchangeResultsContainer> ResultsContainer = nullptr;
+			FString SourceFilename;
 			FString ResultFilepath;
-			TArray<FString> JsonLoadMessages;
 			TMap<FString, FString> ResultPayloads;
 			TUniquePtr<UE::Interchange::Private::FFbxParser> FbxParserPrivate;
 		};
diff --git a/Engine/Plugins/Runtime/Interchange/Source/Pipelines/Private/InterchangeGenericAssetsPipeline.cpp b/Engine/Plugins/Runtime/Interchange/Source/Pipelines/Private/InterchangeGenericAssetsPipeline.cpp
--- a/Engine/Plugins/Runtime/Interchange/Source/Pipelines/Private/InterchangeGenericAssetsPipeline.cpp
+++ b/Engine/Plugins/Runtime/Interchange/Source/Pipelines/Private/InterchangeGenericAssetsPipeline.cpp
@@ -15,6 +15,8 @@
 #include "InterchangeSkeletalMeshLodDataNode.h"
 #include "InterchangeSkeletonFactoryNode.h"
 #include "InterchangeSourceData.h"
+#include "InterchangeStaticMeshFactoryNode.h"
+#include "InterchangeStaticMeshLodDataNode.h"
 #include "InterchangeTexture2DArrayFactoryNode.h"
 #include "InterchangeTexture2DArrayNode.h"
 #include "InterchangeTexture2DNode.h"
@@ -87,6 +89,8 @@ namespace UE::Interchange::Private
 
 bool UInterchangeGenericAssetsPipeline::ExecutePreImportPipeline(UInterchangeBaseNodeContainer* InBaseNodeContainer, const TArray<UInterchangeSourceData*>& InSourceDatas)
 {
+	// @TODO: break up this method into multiple smaller ones, before it becomes more of a monster!
+
 	if (!InBaseNodeContainer)
 	{
 		UE_LOG(LogInterchangePipeline, Warning, TEXT("UInterchangeGenericAssetsPipeline: Cannot execute pre-import pipeline because InBaseNodeContrainer is null"));
@@ -176,7 +180,7 @@ bool UInterchangeGenericAssetsPipeline::ExecutePreImportPipeline(UInterchangeBas
 			if (DependencyIndex > 0)
 			{
 				const FString SkeletalMeshFactoryNodeDependencyUid = SkeletalMeshFactoryDependencyOrder[DependencyIndex - 1];
-				SkeletalMeshFactoryNode->SetFactoryDependencyUid(SkeletalMeshFactoryNodeDependencyUid);
+				SkeletalMeshFactoryNode->AddFactoryDependencyUid(SkeletalMeshFactoryNodeDependencyUid);
 			}
 		};
 
@@ -326,23 +330,144 @@ bool UInterchangeGenericAssetsPipeline::ExecutePreImportPipeline(UInterchangeBas
 		}
 	}
 
-	if (bImportStaticMeshes)
+	if (bImportStaticMeshes && StaticMeshNodes.Num() > 0)
 	{
-		//TODO Import staticmesh here
+		if (bCombineStaticMeshes)
+		{
+			// Combine all the static meshes
+
+			bool bFoundMeshes = false;
+			if (bBakeMeshes)
+			{
+				// If baking transforms, get all the static mesh instance nodes, and group them by LOD
+				TArray<FString> MeshUids;
+				PipelineMeshesUtilities->GetAllStaticMeshInstance(MeshUids);
+
+				TMap<int32, TArray<FString>> MeshUidsPerLodIndex;
+
+				for (const FString& MeshUid : MeshUids)
+				{
+					const FInterchangeMeshInstance& MeshInstance = PipelineMeshesUtilities->GetMeshInstanceByUid(MeshUid);
+					for (const auto& LodIndexAndSceneNodeContainer : MeshInstance.SceneNodePerLodIndex)
+					{
+						const int32 LodIndex = LodIndexAndSceneNodeContainer.Key;
+						const FInterchangeLodSceneNodeContainer& SceneNodeContainer = LodIndexAndSceneNodeContainer.Value;
+
+						TArray<FString>& TranslatedNodes = MeshUidsPerLodIndex.FindOrAdd(LodIndex);
+						for (const UInterchangeSceneNode* SceneNode : SceneNodeContainer.SceneNodes)
+						{
+							TranslatedNodes.Add(SceneNode->GetUniqueID());
+						}
+					}
+				}
+
+				// If we got some instances, create a static mesh factory node
+				if (MeshUidsPerLodIndex.Num() > 0)
+				{
+					UInterchangeStaticMeshFactoryNode* StaticMeshFactoryNode = CreateStaticMeshFactoryNode(MeshUidsPerLodIndex);
+					StaticMeshFactoryNodes.Add(StaticMeshFactoryNode);
+					bFoundMeshes = true;
+				}
+			}
+
+			if (!bFoundMeshes)
+			{
+				// If we haven't yet managed to build a factory node, look at static mesh geometry directly.
+				TArray<FString> MeshUids;
+				PipelineMeshesUtilities->GetAllStaticMeshGeometry(MeshUids);
+
+				TMap<int32, TArray<FString>> MeshUidsPerLodIndex;
+
+				for (const FString& MeshUid : MeshUids)
+				{
+					// MeshGeometry cannot have Lod since LODs are defined in the scene node
+					const FInterchangeMeshGeometry& MeshGeometry = PipelineMeshesUtilities->GetMeshGeometryByUid(MeshUid);
+					const int32 LodIndex = 0;
+					TArray<FString>& TranslatedNodes = MeshUidsPerLodIndex.FindOrAdd(LodIndex);
+					TranslatedNodes.Add(MeshGeometry.MeshUid);
+				}
+
+				if (MeshUidsPerLodIndex.Num() > 0)
+				{
+					UInterchangeStaticMeshFactoryNode* StaticMeshFactoryNode = CreateStaticMeshFactoryNode(MeshUidsPerLodIndex);
+					StaticMeshFactoryNodes.Add(StaticMeshFactoryNode);
+				}
+			}
+
+		}
+		else
+		{
+			// Do not combine static meshes
+
+			bool bFoundMeshes = false;
+			if (bBakeMeshes)
+			{
+				TArray<FString> MeshUids;
+				PipelineMeshesUtilities->GetAllStaticMeshInstance(MeshUids);
+				
+				for (const FString& MeshUid : MeshUids)
+				{
+					// @TODO: check this is correct. What happens when we explicitly don't combine scenes with a LOD group node and multiple transform nodes?
+					TMap<int32, TArray<FString>> MeshUidsPerLodIndex;
+
+					const FInterchangeMeshInstance& MeshInstance = PipelineMeshesUtilities->GetMeshInstanceByUid(MeshUid);
+					for (const auto& LodIndexAndSceneNodeContainer : MeshInstance.SceneNodePerLodIndex)
+					{
+						const int32 LodIndex = LodIndexAndSceneNodeContainer.Key;
+						const FInterchangeLodSceneNodeContainer& SceneNodeContainer = LodIndexAndSceneNodeContainer.Value;
+
+						TArray<FString>& TranslatedNodes = MeshUidsPerLodIndex.FindOrAdd(LodIndex);
+						for (const UInterchangeSceneNode* SceneNode : SceneNodeContainer.SceneNodes)
+						{
+							TranslatedNodes.Add(SceneNode->GetUniqueID());
+						}
+					}
+
+					if (MeshUidsPerLodIndex.Num() > 0)
+					{
+						UInterchangeStaticMeshFactoryNode* StaticMeshFactoryNode = CreateStaticMeshFactoryNode(MeshUidsPerLodIndex);
+						StaticMeshFactoryNodes.Add(StaticMeshFactoryNode);
+						bFoundMeshes = true;
+					}
+				}
+			}
+
+			if (!bFoundMeshes)
+			{
+				TArray<FString> MeshUids;
+				PipelineMeshesUtilities->GetAllStaticMeshGeometry(MeshUids);
+
+				for (const FString& MeshUid : MeshUids)
+				{
+					TMap<int32, TArray<FString>> MeshUidsPerLodIndex;
+
+					const FInterchangeMeshGeometry& MeshGeometry = PipelineMeshesUtilities->GetMeshGeometryByUid(MeshUid);
+					const int32 LodIndex = 0;
+					TArray<FString>& TranslatedNodes = MeshUidsPerLodIndex.FindOrAdd(LodIndex);
+					TranslatedNodes.Add(MeshGeometry.MeshUid);
+
+					if (MeshUidsPerLodIndex.Num() > 0)
+					{
+						UInterchangeStaticMeshFactoryNode* StaticMeshFactoryNode = CreateStaticMeshFactoryNode(MeshUidsPerLodIndex);
+						StaticMeshFactoryNodes.Add(StaticMeshFactoryNode);
+					}
+				}
+			}
+		}
 	}
 
 	const UClass* SkeletalMeshFactoryNodeClass = UInterchangeSkeletalMeshFactoryNode::StaticClass();
 	TArray<FString> SkeletalMeshNodeUids;
 	BaseNodeContainer->GetNodes(SkeletalMeshFactoryNodeClass, SkeletalMeshNodeUids);
 
-// 	const UClass* StaticMeshFactoryNodeClass = UInterchangeStaticMeshFactoryNode::StaticClass();
-// 	TArray<FString> StaticMeshNodeUids;
-// 	BaseNodeContainer->GetNodes(StaticMeshFactoryNodeClass, StaticMeshNodeUids);
+ 	const UClass* StaticMeshFactoryNodeClass = UInterchangeStaticMeshFactoryNode::StaticClass();
+ 	TArray<FString> StaticMeshNodeUids;
+ 	BaseNodeContainer->GetNodes(StaticMeshFactoryNodeClass, StaticMeshNodeUids);
 
 	//TODO count also the imported animations
 
 	//If we import only one asset, and bUseSourceNameForAsset is true, we want to rename the asset using the file name.
-	const int32 MeshesAndAnimsImportedNodeCount = SkeletalMeshNodeUids.Num(); // + StaticMeshNodeUids.Num();
+	const int32 MeshesAndAnimsImportedNodeCount = SkeletalMeshNodeUids.Num() + StaticMeshNodeUids.Num();
 	if (bUseSourceNameForAsset && MeshesAndAnimsImportedNodeCount == 1)
 	{
 		if (SkeletalMeshNodeUids.Num() > 0)
@@ -351,6 +476,12 @@ bool UInterchangeGenericAssetsPipeline::ExecutePreImportPipeline(UInterchangeBas
 			const FString DisplayLabelName = FPaths::GetBaseFilename(SourceDatas[0]->GetFilename());
 			SkeletalMeshNode->SetDisplayLabel(DisplayLabelName);
 		}
+		else if (StaticMeshNodeUids.Num() > 0)
+		{
+			UInterchangeStaticMeshFactoryNode* StaticMeshNode = Cast<UInterchangeStaticMeshFactoryNode>(BaseNodeContainer->GetNode(StaticMeshNodeUids[0]));
+			const FString DisplayLabelName = FPaths::GetBaseFilename(SourceDatas[0]->GetFilename());
+			StaticMeshNode->SetDisplayLabel(DisplayLabelName);
+		}
 	}
 
 	return true;
@@ -433,11 +564,12 @@ UInterchangeMaterialFactoryNode* UInterchangeGenericAssetsPipeline::CreateMateri
 				//We found a factory dependency, because we need to be able to retrieve the UObject create by the factory
 				MaterialFactoryNode->SetTextureDependencyUid(FactoryDepUid);
 				//We also add it to the base dependencies so the texture factory will be called before this material factory
-				MaterialFactoryNode->SetFactoryDependencyUid(FactoryDepUid);
+				MaterialFactoryNode->AddFactoryDependencyUid(FactoryDepUid);
 			}
 		}
 		BaseNodeContainer->AddNode(MaterialFactoryNode);
 		MaterialFactoryNodes.Add(MaterialFactoryNode);
+		MaterialFactoryNode->AddTargetAssetUid(MaterialNode->GetUniqueID());
 	}
 	return MaterialFactoryNode;
 }
@@ -472,12 +604,73 @@ UInterchangeSkeletonFactoryNode* UInterchangeGenericAssetsPipeline::CreateSkelet
 		}
 		SkeletonFactoryNode->InitializeSkeletonNode(SkeletonUid, DisplayLabel, USkeleton::StaticClass()->GetName());
 		SkeletonFactoryNode->SetCustomRootJointUid(RootJointNode->GetUniqueID());
+		SkeletonFactoryNode->AddTargetAssetUid(RootJointNode->GetUniqueID());	// todo: add all joints as target assets?
 		BaseNodeContainer->AddNode(SkeletonFactoryNode);
 	}
 	return SkeletonFactoryNode;
 }
 
-UInterchangeSkeletalMeshFactoryNode* UInterchangeGenericAssetsPipeline::CreateSkeletalMeshFactoryNode(const FString& RootJointUid, const TMap<int32, TArray<FString>> MeshUidsPerLodIndex)
+
+bool UInterchangeGenericAssetsPipeline::MakeMeshFactoryNodeUidAndDisplayLabel(const TMap<int32, TArray<FString>>& MeshUidsPerLodIndex, int32 LodIndex, FString& NewNodeUid, FString& DisplayLabel)
+{
+	int32 SceneNodeCount = 0;
+
+	if (!ensure(LodIndex >= 0 && MeshUidsPerLodIndex.Num() > LodIndex))
+	{
+		return false;
+	}
+
+	for (const TPair<int32, TArray<FString>>& LodIndexAndMeshUids : MeshUidsPerLodIndex)
+	{
+		if (LodIndex == LodIndexAndMeshUids.Key)
+		{
+			const TArray<FString>& Uids = LodIndexAndMeshUids.Value;
+			if (Uids.Num() > 0)
+			{
+				const FString& Uid = Uids[0];
+				UInterchangeBaseNode* Node = BaseNodeContainer->GetNode(Uid);
+
+				if (UInterchangeMeshNode* MeshNode = Cast<UInterchangeMeshNode>(Node))
+				{
+					DisplayLabel = Node->GetDisplayLabel();
+					NewNodeUid = Uid;
+					return true;
+				}
+
+				if (UInterchangeSceneNode* SceneNode = Cast<UInterchangeSceneNode>(Node))
+				{
+					FString RefMeshUid;
+					if (SceneNode->GetCustomMeshDependencyUid(RefMeshUid))
+					{
+						UInterchangeBaseNode* MeshNode = BaseNodeContainer->GetNode(RefMeshUid);
+						if (MeshNode)
+						{
+							DisplayLabel = MeshNode->GetDisplayLabel();
+							if (Uids.Num() == 1)
+							{
+								// If we are instancing one scene node, we want to use it to name the mesh
+								DisplayLabel = SceneNode->GetDisplayLabel();
+							}
+
+							// Use the first scene node uid this LOD references, adding backslash since this uid is not asset typed (\\Mesh\\) like a mesh node
+							// @TODO: change this so that scene nodes get a type prefix of their own?
+							NewNodeUid = TEXT("\\") + Uid;
+							return true;
+						}
+					}
+				}
+			}
+
+			// We found the lod but there is no valid Mesh node to return the Uid
+			break;
+		}
+	}
+
+	return false;
+}
+
+
+UInterchangeSkeletalMeshFactoryNode* UInterchangeGenericAssetsPipeline::CreateSkeletalMeshFactoryNode(const FString& RootJointUid, const TMap<int32, TArray<FString>>& MeshUidsPerLodIndex)
 {
 	//Get the skeleton factory node
 	const UInterchangeBaseNode* RootJointNode = BaseNodeContainer->GetNode(RootJointUid);
@@ -498,80 +691,25 @@ UInterchangeSkeletalMeshFactoryNode* UInterchangeGenericAssetsPipeline::CreateSk
 		return nullptr;
 	}
 	
-	auto GetFirstNodeInfo = [this, &MeshUidsPerLodIndex](const int32 Index, FString& OutFirstMeshNodeUid, int32& OutSceneNodeCount)->const UInterchangeBaseNode*
-	{
-		OutSceneNodeCount = 0;
-		if (!ensure(Index >= 0 && MeshUidsPerLodIndex.Num() > Index))
-		{
-			//Log an error
-			return nullptr;
-		}
-		for (const TPair<int32, TArray<FString>>& LodIndexAndMeshUids : MeshUidsPerLodIndex)
-		{
-			if (Index == LodIndexAndMeshUids.Key)
-			{
-				const TArray<FString>& MeshUids = LodIndexAndMeshUids.Value;
-				if (MeshUids.Num() > 0)
-				{
-					const FString& MeshUid = MeshUids[0];
-					UInterchangeMeshNode* MeshNode = Cast<UInterchangeMeshNode>(BaseNodeContainer->GetNode(MeshUid));
-					if (MeshNode)
-					{
-						OutFirstMeshNodeUid = MeshUid;
-						return MeshNode;
-					}
-					UInterchangeSceneNode* SceneNode = Cast<UInterchangeSceneNode>(BaseNodeContainer->GetNode(MeshUid));
-					if (SceneNode)
-					{
-						FString MeshNodeUid;
-						if (SceneNode->GetCustomMeshDependencyUid(MeshNodeUid))
-						{
-							OutSceneNodeCount = MeshUids.Num();
-							OutFirstMeshNodeUid = MeshNodeUid;
-							return SceneNode;
-						}
-					}
-				}
-				//We found the lod but there is no valid Mesh node to return the Uid
-				break;
-			}
-		}
-		return nullptr;
-	};
-
-	FString FirstMeshNodeUid;
+	// Create the skeletal mesh factory node, name it according to the first mesh node compositing the meshes
+	FString SkeletalMeshUid_MeshNamePart;
+	FString DisplayLabel;
 	const int32 BaseLodIndex = 0;
-	int32 SceneNodeCount = 0;
-	const UInterchangeBaseNode* InterchangeBaseNode = GetFirstNodeInfo(BaseLodIndex, FirstMeshNodeUid, SceneNodeCount);
-	if (!InterchangeBaseNode)
+	if (!MakeMeshFactoryNodeUidAndDisplayLabel(MeshUidsPerLodIndex, BaseLodIndex, SkeletalMeshUid_MeshNamePart, DisplayLabel))
 	{
-		//Log an error
+		// Log an error
 		return nullptr;
 	}
-	const UInterchangeSceneNode* FirstSceneNode = Cast<UInterchangeSceneNode>(InterchangeBaseNode);
-	const UInterchangeMeshNode* FirstMeshNode = Cast<UInterchangeMeshNode>(BaseNodeContainer->GetNode(FirstMeshNodeUid));
 
-	//Create the skeletal mesh factory node, name it according to the first mesh node compositing the meshes
-	FString DisplayLabel = FirstMeshNode->GetDisplayLabel();
-	FString SkeletalMeshUid_MeshNamePart = FirstMeshNodeUid;
-	if(FirstSceneNode)
-	{
-		//If we are instancing one scene node, we want to use it to name the mesh
-		if (SceneNodeCount == 1)
-		{
-			DisplayLabel = FirstSceneNode->GetDisplayLabel();
-		}
-		//Use the first scene node uid this skeletalmesh reference, add backslash since this uid is not asset typed (\\Mesh\\) like FirstMeshNodeUid
-		SkeletalMeshUid_MeshNamePart = TEXT("\\") + FirstSceneNode->GetUniqueID();
-	}
 	const FString SkeletalMeshUid = TEXT("\\SkeletalMesh") + SkeletalMeshUid_MeshNamePart + SkeletonUid;
 	UInterchangeSkeletalMeshFactoryNode* SkeletalMeshFactoryNode = NewObject<UInterchangeSkeletalMeshFactoryNode>(BaseNodeContainer, NAME_None);
 	if (!ensure(SkeletalMeshFactoryNode))
 	{
 		return nullptr;
 	}
+
 	SkeletalMeshFactoryNode->InitializeSkeletalMeshNode(SkeletalMeshUid, DisplayLabel, USkeletalMesh::StaticClass()->GetName());
-	SkeletalMeshFactoryNode->SetFactoryDependencyUid(SkeletonUid);
+	SkeletalMeshFactoryNode->AddFactoryDependencyUid(SkeletonUid);
 	BaseNodeContainer->AddNode(SkeletalMeshFactoryNode);
 
 	AddLodDataToSkeletalMesh(SkeletonFactoryNode, SkeletalMeshFactoryNode, MeshUidsPerLodIndex);
@@ -579,6 +717,7 @@ UInterchangeSkeletalMeshFactoryNode* UInterchangeGenericAssetsPipeline::CreateSk
 	return SkeletalMeshFactoryNode;
 }
 
+
 UInterchangeSkeletalMeshLodDataNode* UInterchangeGenericAssetsPipeline::CreateSkeletalMeshLodDataNode(const FString& NodeName, const FString& NodeUniqueID)
 {
 	FString DisplayLabel(NodeName);
@@ -589,12 +728,13 @@ UInterchangeSkeletalMeshLodDataNode* UInterchangeGenericAssetsPipeline::CreateSk
 		//TODO log error
 		return nullptr;
 	}
-	// Creating a UMaterialInterface
+
 	SkeletalMeshLodDataNode->InitializeNode(NodeUID, DisplayLabel, EInterchangeNodeContainerType::NodeContainerType_FactoryData);
 	BaseNodeContainer->AddNode(SkeletalMeshLodDataNode);
 	return SkeletalMeshLodDataNode;
 }
 
+
 void UInterchangeGenericAssetsPipeline::AddLodDataToSkeletalMesh(const UInterchangeSkeletonFactoryNode* SkeletonFactoryNode, UInterchangeSkeletalMeshFactoryNode* SkeletalMeshFactoryNode, const TMap<int32, TArray<FString>>& NodeUidsPerLodIndex)
 {
 	const FString SkeletalMeshUid = SkeletalMeshFactoryNode->GetUniqueID();
@@ -630,10 +770,18 @@ void UInterchangeGenericAssetsPipeline::AddLodDataToSkeletalMesh(const UIntercha
 			TArray<FString> MaterialDependencies;
 			if (const UInterchangeSceneNode* SceneNode = Cast<UInterchangeSceneNode>(BaseNodeContainer->GetNode(NodeUid)))
 			{
+				FString MeshDependency;
+				SceneNode->GetCustomMeshDependencyUid(MeshDependency);
+				if (BaseNodeContainer->IsNodeUidValid(MeshDependency))
+				{
+					SkeletalMeshFactoryNode->AddTargetAssetUid(MeshDependency);
+				}
+
 				SceneNode->GetMaterialDependencyUids(MaterialDependencies);
 			}
 			else if (const UInterchangeMeshNode* MeshNode = Cast<UInterchangeMeshNode>(BaseNodeContainer->GetNode(NodeUid)))
 			{
+				SkeletalMeshFactoryNode->AddTargetAssetUid(NodeUid);
 				MeshNode->GetMaterialDependencies(MaterialDependencies);
 			}
 			const int32 MaterialCount = MaterialDependencies.Num();
@@ -647,7 +795,7 @@ void UInterchangeGenericAssetsPipeline::AddLodDataToSkeletalMesh(const UIntercha
 					SkeletalMeshFactoryNode->GetFactoryDependencies(FactoryDependencies);
 					if (!FactoryDependencies.Contains(MaterialFactoryNodeUid))
 					{
-						SkeletalMeshFactoryNode->SetFactoryDependencyUid(MaterialFactoryNodeUid);
+						SkeletalMeshFactoryNode->AddFactoryDependencyUid(MaterialFactoryNodeUid);
 					}
 				}
 			}
@@ -655,3 +803,125 @@ void UInterchangeGenericAssetsPipeline::AddLodDataToSkeletalMesh(const UIntercha
 		}
 	}
 }
+
+
+UInterchangeStaticMeshFactoryNode* UInterchangeGenericAssetsPipeline::CreateStaticMeshFactoryNode(const TMap<int32, TArray<FString>>& MeshUidsPerLodIndex)
+{
+	if (MeshUidsPerLodIndex.Num() == 0)
+	{
+		return nullptr;
+	}
+
+	// Create the static mesh factory node, name it according to the first mesh node compositing the meshes
+	FString StaticMeshUid_MeshNamePart;
+	FString DisplayLabel;
+	const int32 BaseLodIndex = 0;
+	if (!MakeMeshFactoryNodeUidAndDisplayLabel(MeshUidsPerLodIndex, BaseLodIndex, StaticMeshUid_MeshNamePart, DisplayLabel))
+	{
+		// Log an error
+		return nullptr;
+	}
+
+	const FString StaticMeshUid = TEXT("\\StaticMesh") + StaticMeshUid_MeshNamePart;
+	UInterchangeStaticMeshFactoryNode* StaticMeshFactoryNode = NewObject<UInterchangeStaticMeshFactoryNode>(BaseNodeContainer, NAME_None);
+	if (!ensure(StaticMeshFactoryNode))
+	{
+		return nullptr;
+	}
+
+	StaticMeshFactoryNode->InitializeStaticMeshNode(StaticMeshUid, DisplayLabel, UStaticMesh::StaticClass()->GetName());
+	BaseNodeContainer->AddNode(StaticMeshFactoryNode);
+
+	AddLodDataToStaticMesh(StaticMeshFactoryNode, MeshUidsPerLodIndex);
+
+	return StaticMeshFactoryNode;
+}
+
+
+UInterchangeStaticMeshLodDataNode* UInterchangeGenericAssetsPipeline::CreateStaticMeshLodDataNode(const FString& NodeName, const FString& NodeUniqueID)
+{
+	FString DisplayLabel(NodeName);
+	FString NodeUID(NodeUniqueID);
+	UInterchangeStaticMeshLodDataNode* StaticMeshLodDataNode = NewObject<UInterchangeStaticMeshLodDataNode>(BaseNodeContainer, NAME_None);
+	if (!ensure(StaticMeshLodDataNode))
+	{
+		// @TODO: log error
+		return nullptr;
+	}
+
+	StaticMeshLodDataNode->InitializeNode(NodeUID, DisplayLabel, EInterchangeNodeContainerType::NodeContainerType_FactoryData);
+	BaseNodeContainer->AddNode(StaticMeshLodDataNode);
+	return StaticMeshLodDataNode;
+}
+
+
+void UInterchangeGenericAssetsPipeline::AddLodDataToStaticMesh(UInterchangeStaticMeshFactoryNode* StaticMeshFactoryNode, const TMap<int32, TArray<FString>>& NodeUidsPerLodIndex)
+{
+	const FString StaticMeshFactoryUid = StaticMeshFactoryNode->GetUniqueID();
+
+	for (const TPair<int32, TArray<FString>>& LodIndexAndNodeUids : NodeUidsPerLodIndex)
+	{
+		const int32 LodIndex = LodIndexAndNodeUids.Key;
+		if (!bImportLods && LodIndex > 0)
+		{
+			// If the pipeline should not import lods, skip any lod over base lod
+			continue;
+		}
+
+		const TArray<FString>& NodeUids = LodIndexAndNodeUids.Value;
+
+		// Create a lod data node with all the meshes for this LOD
+		const FString StaticMeshLodDataName = TEXT("LodData") + FString::FromInt(LodIndex);
+		const FString LODDataPrefix = TEXT("\\LodData") + (LodIndex > 0 ? FString::FromInt(LodIndex) : TEXT(""));
+		const FString StaticMeshLodDataUniqueID = LODDataPrefix + StaticMeshFactoryUid;
+
+		// Create LodData node if it doesn't already exist
+		UInterchangeStaticMeshLodDataNode* LodDataNode = Cast<UInterchangeStaticMeshLodDataNode>(BaseNodeContainer->GetNode(StaticMeshLodDataUniqueID));
+		if (!LodDataNode)
+		{
+			// Add the data for the LOD (all the mesh node fbx path, so we can find them when we will create the payload data)
+			LodDataNode = CreateStaticMeshLodDataNode(StaticMeshLodDataName, StaticMeshLodDataUniqueID);
+			LodDataNode->SetParentUid(StaticMeshFactoryUid);
+			StaticMeshFactoryNode->AddLodDataUniqueId(StaticMeshLodDataUniqueID);
+		}
+
+		for (const FString& NodeUid : NodeUids)
+		{
+			TArray<FString> MaterialDependencies;
+			if (const UInterchangeSceneNode* SceneNode = Cast<UInterchangeSceneNode>(BaseNodeContainer->GetNode(NodeUid)))
+			{
+				FString MeshDependency;
+				SceneNode->GetCustomMeshDependencyUid(MeshDependency);
+				if (BaseNodeContainer->IsNodeUidValid(MeshDependency))
+				{
+					StaticMeshFactoryNode->AddTargetAssetUid(MeshDependency);
+				}
+				
+				SceneNode->GetMaterialDependencyUids(MaterialDependencies);
+			}
+			else if (const UInterchangeMeshNode* MeshNode = Cast<UInterchangeMeshNode>(BaseNodeContainer->GetNode(NodeUid)))
+			{
+				StaticMeshFactoryNode->AddTargetAssetUid(NodeUid);
+				MeshNode->GetMaterialDependencies(MaterialDependencies);
+			}
+
+			const int32 MaterialCount = MaterialDependencies.Num();
+			for (int32 MaterialIndex = 0; MaterialIndex < MaterialCount; ++MaterialIndex)
+			{
+				const FString MaterialFactoryNodeUid = UInterchangeMaterialFactoryNode::GetMaterialFactoryNodeUidFromMaterialNodeUid(MaterialDependencies[MaterialIndex]);
+				if (BaseNodeContainer->IsNodeUidValid(MaterialFactoryNodeUid))
+				{
+					// Create a factory dependency so Material asset are imported before the skeletal mesh asset
+					TArray<FString> FactoryDependencies;
+					StaticMeshFactoryNode->GetFactoryDependencies(FactoryDependencies);
+					if (!FactoryDependencies.Contains(MaterialFactoryNodeUid))
+					{
+						StaticMeshFactoryNode->AddFactoryDependencyUid(MaterialFactoryNodeUid);
+					}
+				}
+			}
+
+			LodDataNode->AddMeshUid(NodeUid);
+		}
+	}
+}
diff --git a/Engine/Plugins/Runtime/Interchange/Source/Pipelines/Public/InterchangeGenericAssetsPipeline.h b/Engine/Plugins/Runtime/Interchange/Source/Pipelines/Public/InterchangeGenericAssetsPipeline.h
--- a/Engine/Plugins/Runtime/Interchange/Source/Pipelines/Public/InterchangeGenericAssetsPipeline.h
+++ b/Engine/Plugins/Runtime/Interchange/Source/Pipelines/Public/InterchangeGenericAssetsPipeline.h
@@ -15,10 +15,13 @@ class UInterchangeMaterialFactoryNode;
 class UInterchangeMaterialNode;
 class UInterchangeMeshNode;
 class UInterchangePipelineMeshesUtilities;
+class UInterchangeResult;
 class UInterchangeSceneNode;
 class UInterchangeSkeletalMeshFactoryNode;
 class UInterchangeSkeletalMeshLodDataNode;
 class UInterchangeSkeletonFactoryNode;
+class UInterchangeStaticMeshFactoryNode;
+class UInterchangeStaticMeshLodDataNode;
 class UInterchangeTexture2DArrayFactoryNode;
 class UInterchangeTexture2DArrayNode;
 class UInterchangeTextureCubeFactoryNode;
@@ -188,14 +191,11 @@ private:
 	/** Skeletal mesh factory assets nodes */
 	TArray<UInterchangeSkeletalMeshFactoryNode*> SkeletalMeshFactoryNodes;
 	
-	/** Static mesh factory assets nodes */
-	//TArray<UInterchangeStaticMeshFactoryNode*> StaticMeshFactoryNodes;
-	
 	/**
 	 * This function can create a UInterchangeSkeletalMeshFactoryNode
 	 * @param MeshUidsPerLodIndex - The MeshUids can represent a SceneNode pointing on a MeshNode or directly a MeshNode
 	 */
-	UInterchangeSkeletalMeshFactoryNode* CreateSkeletalMeshFactoryNode(const FString& RootJointUid, const TMap<int32, TArray<FString>> MeshUidsPerLodIndex);
+	UInterchangeSkeletalMeshFactoryNode* CreateSkeletalMeshFactoryNode(const FString& RootJointUid, const TMap<int32, TArray<FString>>& MeshUidsPerLodIndex);
 	
 	/** This function can create a UInterchangeSkeletalMeshLodDataNode which represent the LOD data need by the factory to create a lod mesh */
 	UInterchangeSkeletalMeshLodDataNode* CreateSkeletalMeshLodDataNode(const FString& NodeName, const FString& NodeUniqueID);
@@ -206,6 +206,30 @@ private:
 	 */
 	void AddLodDataToSkeletalMesh(const UInterchangeSkeletonFactoryNode* SkeletonFactoryNode, UInterchangeSkeletalMeshFactoryNode* SkeletalMeshFactoryNode, const TMap<int32, TArray<FString>>& NodeUidsPerLodIndex);
 
+
+	/** Static mesh factory assets nodes */
+	TArray<UInterchangeStaticMeshFactoryNode*> StaticMeshFactoryNodes;
+
+	/**
+	 * This function can create a UInterchangeStaticMeshFactoryNode
+	 * @param MeshUidsPerLodIndex - The MeshUids can represent a SceneNode pointing on a MeshNode or directly a MeshNode
+	 */
+	UInterchangeStaticMeshFactoryNode* CreateStaticMeshFactoryNode(const TMap<int32, TArray<FString>>& MeshUidsPerLodIndex);
+	
+	/** This function can create a UInterchangeStaticMeshLodDataNode which represents the LOD data needed by the factory to create a lod mesh */
+	UInterchangeStaticMeshLodDataNode* CreateStaticMeshLodDataNode(const FString& NodeName, const FString& NodeUniqueID);
+
+	/**
+	 * This function add all lod data nodes to the static mesh.
+	 * @param NodeUidsPerLodIndex - The NodeUids can be a UInterchangeSceneNode or a UInterchangeMeshNode. The scene node can bake each instance of the mesh versus the mesh node will import only the modelled mesh.
+	 */
+	void AddLodDataToStaticMesh(UInterchangeStaticMeshFactoryNode* StaticMeshFactoryNode, const TMap<int32, TArray<FString>>& NodeUidsPerLodIndex);
+
+	/**
+	 * Return a reasonable UID and display label for a new mesh factory node.
+	 */
+	bool MakeMeshFactoryNodeUidAndDisplayLabel(const TMap<int32, TArray<FString>>& MeshUidsPerLodIndex, int32 LodIndex, FString& NewMeshUid, FString& DisplayLabel);
+
 	/** Translated scene nodes */
 	TArray<UInterchangeSceneNode*> SceneNodes;
 
diff --git a/Engine/Source/Developer/AssetTools/Private/AssetTools.cpp b/Engine/Source/Developer/AssetTools/Private/AssetTools.cpp
--- a/Engine/Source/Developer/AssetTools/Private/AssetTools.cpp
+++ b/Engine/Source/Developer/AssetTools/Private/AssetTools.cpp
@@ -1937,6 +1937,20 @@ TArray<UObject*> UAssetToolsImpl::ImportAssetsInternal(const TArray<FString>& Fi
 			SlowTask.CompletedWork = FilesAndDestinations.Num();
 		}
 	}
+
+	struct FInterchangeImportStatus
+	{
+		explicit FInterchangeImportStatus(int32 NumFiles)
+			: InterchangeResultsContainer(NewObject<UInterchangeResultsContainer>(GetTransientPackage())),
+			  ImportCount(NumFiles)
+		{}
+
+		TStrongObjectPtr<UInterchangeResultsContainer> InterchangeResultsContainer;
+		std::atomic<int32> ImportCount;
+	};
+
+	TSharedPtr<FInterchangeImportStatus, ESPMode::ThreadSafe> ImportStatus = MakeShared<FInterchangeImportStatus>(FilesAndDestinations.Num());
+
 	// Now iterate over the input files and use the same factory object for each file with the same extension
 	for(int32 FileIdx = 0; FileIdx < FilesAndDestinations.Num() && !bImportWasCancelled; ++FileIdx)
 	{
@@ -1965,7 +1979,19 @@ TArray<UObject*> UAssetToolsImpl::ImportAssetsInternal(const TArray<FString>& Fi
 				ImportAssetParameters.bIsAutomated = bAutomatedImport;
 				ImportAssetParameters.ReimportAsset = nullptr;
 
-				InterchangeManager.ImportAsset(DestinationPath, ScopedSourceData.GetSourceData(), ImportAssetParameters);
+				UE::Interchange::FAssetImportResultRef InterchangeResult = (InterchangeManager.ImportAssetAsync(DestinationPath, ScopedSourceData.GetSourceData(), ImportAssetParameters));
+				InterchangeResult->OnDone(
+					// Note: ImportStatus captured by value so that the lambda keeps the shared ptr alive
+					[ImportStatus](UE::Interchange::FAssetImportResult& Result)
+					{
+						ImportStatus->InterchangeResultsContainer->Append(Result.GetResults());
+						if (--ImportStatus->ImportCount == 0)
+						{
+							UInterchangeManager& InterchangeManager = UInterchangeManager::GetInterchangeManager();
+							InterchangeManager.OnBatchImportComplete.Broadcast(ImportStatus->InterchangeResultsContainer);
+						}
+					});
+
 				InterchangeManager.OnAssetPostImport.Remove(PostImportHandle);
 				//Import done, iterate the next file and destination
 				
diff --git a/Engine/Source/Runtime/Engine/Classes/Engine/StaticMesh.h b/Engine/Source/Runtime/Engine/Classes/Engine/StaticMesh.h
--- a/Engine/Source/Runtime/Engine/Classes/Engine/StaticMesh.h
+++ b/Engine/Source/Runtime/Engine/Classes/Engine/StaticMesh.h
@@ -1319,6 +1319,22 @@ public:
 	 */
 	ENGINE_API bool ModifyHiResMeshDescription(bool bAlwaysMarkDirty = true);
 
+	/**
+	 * Get AssetImportData for the static mesh
+	 */
+	class UAssetImportData* GetAssetImportData() const
+	{
+		return AssetImportData;
+	}
+
+	/**
+	 * Set AssetImportData for the static mesh
+	 */
+	void SetAssetImportData(class UAssetImportData* InAssetImportData)
+	{
+		AssetImportData = InAssetImportData;
+	}
+
 	/**
 	 * Adds an empty UV channel at the end of the existing channels on the given LOD of a StaticMesh.
 	 * @param	LODIndex			Index of the StaticMesh LOD.
diff --git a/Engine/Source/Runtime/Interchange/Core/InterchangeCore.Build.cs b/Engine/Source/Runtime/Interchange/Core/InterchangeCore.Build.cs
--- a/Engine/Source/Runtime/Interchange/Core/InterchangeCore.Build.cs
+++ b/Engine/Source/Runtime/Interchange/Core/InterchangeCore.Build.cs
@@ -10,6 +10,7 @@ public class InterchangeCore : ModuleRules
 			new string[] {
 				"Core",
 				"CoreUObject",
+				"Json"
 			}
 		);
 	}
diff --git a/Engine/Source/Runtime/Interchange/Core/Private/Nodes/InterchangeBaseNode.cpp b/Engine/Source/Runtime/Interchange/Core/Private/Nodes/InterchangeBaseNode.cpp
--- a/Engine/Source/Runtime/Interchange/Core/Private/Nodes/InterchangeBaseNode.cpp
+++ b/Engine/Source/Runtime/Interchange/Core/Private/Nodes/InterchangeBaseNode.cpp
@@ -118,7 +118,7 @@ void UInterchangeBaseNode::GetFactoryDependencies(TArray<FString>& OutDependenci
 	FactoryDependencies.GetNames(OutDependencies);
 }
 
-bool UInterchangeBaseNode::SetFactoryDependencyUid(const FString& DependencyUid)
+bool UInterchangeBaseNode::AddFactoryDependencyUid(const FString& DependencyUid)
 {
 	return FactoryDependencies.AddName(DependencyUid);
 }
@@ -211,6 +211,26 @@ bool UInterchangeBaseNode::SetAssetName(const FString& AssetName)
 	return false;
 }
 
+int32 UInterchangeBaseNode::GetTargetAssetCount() const
+{
+	return TargetAssets.GetCount();
+}
+
+void UInterchangeBaseNode::GetTargetAssetUids(TArray<FString>& OutTargetAssets) const
+{
+	TargetAssets.GetNames(OutTargetAssets);
+}
+
+bool UInterchangeBaseNode::AddTargetAssetUid(const FString& AssetUid)
+{
+	return TargetAssets.AddName(AssetUid);
+}
+
+bool UInterchangeBaseNode::RemoveTargetAssetUid(const FString& AssetUid)
+{
+	return TargetAssets.RemoveName(AssetUid);
+}
+
 FString UInterchangeBaseNode::InvalidNodeUid()
 {
 	return FString();
diff --git a/Engine/Source/Runtime/Interchange/Core/Public/InterchangeFactoryBase.h b/Engine/Source/Runtime/Interchange/Core/Public/InterchangeFactoryBase.h
--- a/Engine/Source/Runtime/Interchange/Core/Public/InterchangeFactoryBase.h
+++ b/Engine/Source/Runtime/Interchange/Core/Public/InterchangeFactoryBase.h
@@ -3,6 +3,7 @@
 #pragma once
 
 #include "CoreMinimal.h"
+#include "InterchangeResultsContainer.h"
 #include "UObject/Class.h"
 #include "UObject/Object.h"
 #include "UObject/ObjectMacros.h"
@@ -126,4 +127,32 @@ public:
 		check(IsInGameThread());
 		return;
 	}
+
+	/**
+	 * This function is used to add the given message object directly into the results for this operation.
+	 */
+	template <typename T>
+	T* AddMessage() const
+	{
+		check(Results != nullptr);
+		T* Item = Results->Add<T>();
+		return Item;
+	}
+
+
+	void AddMessage(UInterchangeResult* Item) const
+	{
+		check(Results != nullptr);
+		Results->Add(Item);
+	}
+	
+
+	void SetResultsContainer(UInterchangeResultsContainer* InResults)
+	{
+		Results = InResults;
+	}
+
+
+	UPROPERTY()
+	UInterchangeResultsContainer* Results;
 };
diff --git a/Engine/Source/Runtime/Interchange/Core/Public/InterchangePipelineBase.h b/Engine/Source/Runtime/Interchange/Core/Public/InterchangePipelineBase.h
--- a/Engine/Source/Runtime/Interchange/Core/Public/InterchangePipelineBase.h
+++ b/Engine/Source/Runtime/Interchange/Core/Public/InterchangePipelineBase.h
@@ -4,6 +4,7 @@
 
 #include "CoreMinimal.h"
 #include "Nodes/InterchangeBaseNodeContainer.h"
+#include "InterchangeResultsContainer.h"
 #include "InterchangeSourceData.h"
 #include "UObject/Object.h"
 #include "UObject/ObjectMacros.h"
@@ -93,6 +94,29 @@ public:
 
 	void SaveSettings(const FName PipelineStackName);
 
+	/**
+	 * This function is used to add the given message object directly into the results for this operation.
+	 */
+	template <typename T>
+	T* AddMessage() const
+	{
+		check(Results != nullptr);
+		T* Item = Results->Add<T>();
+		return Item;
+	}
+
+	void AddMessage(UInterchangeResult* Item) const
+	{
+		check(Results != nullptr);
+		Results->Add(Item);
+	}
+	
+	void SetResultsContainer(UInterchangeResultsContainer* InResults)
+	{
+		Results = InResults;
+	}
+
+
 protected:
 
 	/**
@@ -131,4 +155,8 @@ protected:
 	{
 		return false;
 	}
+
+
+	UPROPERTY()
+	TObjectPtr<UInterchangeResultsContainer> Results;
 };
diff --git a/Engine/Source/Runtime/Interchange/Core/Public/InterchangeTranslatorBase.h b/Engine/Source/Runtime/Interchange/Core/Public/InterchangeTranslatorBase.h
--- a/Engine/Source/Runtime/Interchange/Core/Public/InterchangeTranslatorBase.h
+++ b/Engine/Source/Runtime/Interchange/Core/Public/InterchangeTranslatorBase.h
@@ -4,6 +4,7 @@
 
 #include "CoreMinimal.h"
 #include "Nodes/InterchangeBaseNodeContainer.h"
+#include "InterchangeResultsContainer.h"
 #include "InterchangeSourceData.h"
 #include "UObject/Object.h"
 #include "UObject/ObjectMacros.h"
@@ -16,20 +17,19 @@ class INTERCHANGECORE_API UInterchangeTranslatorBase : public UObject
 	GENERATED_BODY()
 public:
 
-	/** return true if the translator can translate the specified file. */
-	virtual bool CanImportSourceData(const UInterchangeSourceData* SourceData) const
+	/** return true if the translator can translate the given source data. */
+	virtual bool CanImportSourceData(const UInterchangeSourceData* InSourceData) const
 	{
 		return false;
 	}
 
 	/**
-	 * Translate a source data into node(s) that are hold in the specified nodes container.
+	 * Translate the associated source data into node(s) that are hold in the specified nodes container.
 	 *
-	 * @param SourceData - The source data containing the data to translate
 	 * @param BaseNodeContainer - The nodes container where to put the translated source data result.
 	 * @return true if the translator can translate the source data, false otherwise.
 	 */
-	virtual bool Translate(const UInterchangeSourceData* SourceData, UInterchangeBaseNodeContainer& BaseNodeContainer) const
+	virtual bool Translate(UInterchangeBaseNodeContainer& BaseNodeContainer) const
 	{
 		return false;
 	}
@@ -53,4 +53,47 @@ public:
 	{
 		return;
 	}
+
+	/**
+	 * This function is used to add the given message object directly into the results for this operation.
+	 */
+	template <typename T>
+	T* AddMessage() const
+	{
+		check(Results != nullptr);
+		check(SourceData != nullptr);
+		T* Item = Results->Add<T>();
+		Item->SourceAssetName = SourceData->GetFilename();
+		return Item;
+	}
+
+
+	void AddMessage(UInterchangeResult* Item) const
+	{
+		check(Results != nullptr);
+		check(SourceData != nullptr);
+		Results->Add(Item);
+		Item->SourceAssetName = SourceData->GetFilename();
+	}
+	
+
+	void SetResultsContainer(UInterchangeResultsContainer* InResults)
+	{
+		Results = InResults;
+	}
+
+	/**
+	 * Get the associated source data for this translator.
+	 */
+	const UInterchangeSourceData* GetSourceData() const
+	{
+		return SourceData;
+	}
+
+
+	UPROPERTY()
+	TObjectPtr<UInterchangeResultsContainer> Results;
+
+	UPROPERTY()
+	TObjectPtr<const UInterchangeSourceData> SourceData;
 };
diff --git a/Engine/Source/Runtime/Interchange/Core/Public/Nodes/InterchangeBaseNode.h b/Engine/Source/Runtime/Interchange/Core/Public/Nodes/InterchangeBaseNode.h
--- a/Engine/Source/Runtime/Interchange/Core/Public/Nodes/InterchangeBaseNode.h
+++ b/Engine/Source/Runtime/Interchange/Core/Public/Nodes/InterchangeBaseNode.h
@@ -751,6 +751,12 @@ namespace UE
 				return AttributeKey;
 			}
 
+			static const FString& TargetAssetIDsKey()
+			{
+				static FString Key(TEXT("__TARGET_ASSET_IDS_"));
+				return Key;
+			}
+
 			static const FString& GetFactoryDependenciesBaseKey()
 			{
 				static FString BaseNodeFactoryDependencies_BaseKey = TEXT("__BaseNodeFactoryDependencies__");
@@ -803,6 +809,7 @@ public:
 	{
 		Attributes = MakeShared<UE::Interchange::FAttributeStorage, ESPMode::ThreadSafe>();
 		FactoryDependencies.Initialize(Attributes, UE::Interchange::FBaseNodeStaticData::GetFactoryDependenciesBaseKey());
+		TargetAssets.Initialize(Attributes, UE::Interchange::FBaseNodeStaticData::TargetAssetIDsKey());
 		RegisterAttribute<bool>(UE::Interchange::FBaseNodeStaticData::IsEnabledKey(), true);
 		RegisterAttribute<uint8>(UE::Interchange::FBaseNodeStaticData::NodeContainerTypeKey(), static_cast<uint8>(EInterchangeNodeContainerType::NodeContainerType_None));
 	}
@@ -872,6 +879,20 @@ public:
 				KeyDisplayName += NodeAttributeKey.Key.RightChop(IndexPosition);
 			}
 		}
+		else if (NodeAttributeKey.Key.Equals(UE::Interchange::FBaseNodeStaticData::TargetAssetIDsKey()))
+		{
+			KeyDisplayName = TEXT("Target Asset Count");
+		}
+		else if (NodeAttributeKey.Key.StartsWith(UE::Interchange::FBaseNodeStaticData::TargetAssetIDsKey()))
+		{
+			KeyDisplayName = TEXT("Target Asset Index ");
+			const FString IndexKey = UE::Interchange::FNameAttributeArrayHelper::IndexKey();
+			int32 IndexPosition = NodeAttributeKey.Key.Find(IndexKey) + IndexKey.Len();
+			if (IndexPosition < NodeAttributeKey.Key.Len())
+			{
+				KeyDisplayName += NodeAttributeKey.Key.RightChop(IndexPosition);
+			}
+		}
 		else if (NodeAttributeKey == UE::Interchange::FBaseNodeStaticData::NodeContainerTypeKey())
 		{
 			KeyDisplayName = TEXT("Node Container Type");
@@ -1005,7 +1026,7 @@ public:
 	 * Add one dependency to this object.
 	 */
 	UFUNCTION(BlueprintCallable, Category = "Interchange | Node")
-	bool SetFactoryDependencyUid(const FString& DependencyUid);
+	bool AddFactoryDependencyUid(const FString& DependencyUid);
 
 	/**
 	 * Remove one dependency from this object.
@@ -1013,6 +1034,30 @@ public:
 	UFUNCTION(BlueprintCallable, Category = "Interchange | Node")
 	bool RemoveFactoryDependencyUid(const FString& DependencyUid);
 
+	/**
+	 * Get number of target assets relating to this object.
+	 */
+	UFUNCTION(BlueprintCallable, Category = "Interchange | Node")
+	int32 GetTargetAssetCount() const;
+
+	/**
+	 * Get target assets relating to this object.
+	 */
+	UFUNCTION(BlueprintCallable, Category = "Interchange | Node")
+	void GetTargetAssetUids(TArray<FString>& OutTargetAssets) const;
+
+	/**
+	 * Add asset node UID relating to this object.
+	 */
+	UFUNCTION(BlueprintCallable, Category = "Interchange | Node")
+	bool AddTargetAssetUid(const FString& AssetUid);
+
+	/**
+	 * Remove asset node UID relating to this object.
+	 */
+	UFUNCTION(BlueprintCallable, Category = "Interchange | Node")
+	bool RemoveTargetAssetUid(const FString& AssetUid);
+
 	/**
 	 * IsEnable true mean that the node will be import/export, if false it will be discarded.
 	 * Return false if this node was disabled. Return true if the attribute is not there or if it was enabled.
@@ -1110,4 +1155,9 @@ protected:
 	 *          Material factory node will have dependencies on texture factory node
 	 */
 	UE::Interchange::FNameAttributeArrayHelper FactoryDependencies;
+
+	/**
+	 * This tracks the IDs of asset nodes which are the target of factories
+	 */
+	UE::Interchange::FNameAttributeArrayHelper TargetAssets;
 };
diff --git a/Engine/Source/Runtime/Interchange/Engine/Private/InterchangeManager.cpp b/Engine/Source/Runtime/Interchange/Engine/Private/InterchangeManager.cpp
--- a/Engine/Source/Runtime/Interchange/Engine/Private/InterchangeManager.cpp
+++ b/Engine/Source/Runtime/Interchange/Engine/Private/InterchangeManager.cpp
@@ -78,22 +78,10 @@ UE::Interchange::FImportAsyncHelper::FImportAsyncHelper()
 
 void UE::Interchange::FImportAsyncHelper::AddReferencedObjects(FReferenceCollector& Collector)
 {
-	for (UInterchangeSourceData* SourceData : SourceDatas)
-	{
-		Collector.AddReferencedObject(SourceData);
-	}
-	for (UInterchangeTranslatorBase* Translator : Translators)
-	{
-		Collector.AddReferencedObject(Translator);
-	}
-	for (UInterchangePipelineBase* Pipeline : Pipelines)
-	{
-		Collector.AddReferencedObject(Pipeline);
-	}
-	for (UInterchangeFactoryBase* Factory : Factories)
-	{
-		Collector.AddReferencedObject(Factory);
-	}
+	Collector.AddReferencedObjects(SourceDatas);
+	Collector.AddReferencedObjects(Translators);
+	Collector.AddReferencedObjects(Pipelines);
+	Collector.AddReferencedObjects(CreatedFactories);
 }
 
 void UE::Interchange::FImportAsyncHelper::ReleaseTranslatorsSource()
@@ -177,7 +165,6 @@ void UE::Interchange::FImportAsyncHelper::CleanUp()
 		{
 			SourceData->RemoveFromRoot();
 			SourceData->MarkPendingKill();
-			SourceData = nullptr;
 		}
 	}
 	SourceDatas.Empty();
@@ -189,7 +176,6 @@ void UE::Interchange::FImportAsyncHelper::CleanUp()
 			Translator->ImportFinish();
 			Translator->RemoveFromRoot();
 			Translator->MarkPendingKill();
-			Translator = nullptr;
 		}
 	}
 	Translators.Empty();
@@ -200,18 +186,25 @@ void UE::Interchange::FImportAsyncHelper::CleanUp()
 		{
 			Pipeline->RemoveFromRoot();
 			Pipeline->MarkPendingKill();
-			Pipeline = nullptr;
 		}
 	}
 	Pipelines.Empty();
 
-	//Factories are not instantiate, we use the registered one directly
-	Factories.Empty();
+	for (const TPair<FString, UInterchangeFactoryBase*>& FactoryKeyAndValue : CreatedFactories)
+	{
+		if (FactoryKeyAndValue.Value)
+		{
+			FactoryKeyAndValue.Value->RemoveFromRoot();
+			FactoryKeyAndValue.Value->MarkPendingKill();
+		}
+	}
+	CreatedFactories.Empty();
 }
 
 UE::Interchange::FAssetImportResult::FAssetImportResult()
 	: ImportStatus(EStatus::Invalid)
 {
+	Results = NewObject<UInterchangeResultsContainer>(GetTransientPackage());
 }
 
 UE::Interchange::FAssetImportResult::EStatus UE::Interchange::FAssetImportResult::GetStatus() const
@@ -295,6 +288,7 @@ void UE::Interchange::FAssetImportResult::AddReferencedObjects(FReferenceCollect
 {
 	FReadScopeLock ReadScopeLock(ImportedAssetsRWLock);
 	Collector.AddReferencedObjects(ImportedAssets);
+	Collector.AddReferencedObject(Results);
 }
 
 void UE::Interchange::SanitizeObjectPath(FString& ObjectPath)
@@ -365,48 +359,34 @@ UInterchangeManager& UInterchangeManager::GetInterchangeManager()
 
 bool UInterchangeManager::RegisterTranslator(const UClass* TranslatorClass)
 {
-	if(!TranslatorClass)
+	if (!TranslatorClass)
 	{
 		return false;
 	}
 
-	decltype(RegisteredTranslatorsClass)::ElementType TranslatorAsKey(TranslatorClass);
-	uint32 TranslatorClassHash = GetTypeHash(TranslatorAsKey);
-	if(RegisteredTranslatorsClass.ContainsByHash(TranslatorClassHash, TranslatorAsKey))
-	{
-		return true;
-	}
-	RegisteredTranslatorsClass.AddByHash(TranslatorClassHash, TranslatorAsKey);
-
-
-	UInterchangeTranslatorBase* TranslatorToRegister = NewObject<UInterchangeTranslatorBase>(GetTransientPackage(), TranslatorClass, NAME_None);
-	if(!TranslatorToRegister)
-	{
-		return false;
-	}
-	RegisteredTranslators.Add(TranslatorToRegister);
+	RegisteredTranslatorsClass.Add(TranslatorClass);
 	return true;
 }
 
 bool UInterchangeManager::RegisterFactory(const UClass* FactoryClass)
 {
-	if (!FactoryClass)
+	if (!FactoryClass || !FactoryClass->IsChildOf<UInterchangeFactoryBase>())
 	{
 		return false;
 	}
 
-	UInterchangeFactoryBase* FactoryToRegister = NewObject<UInterchangeFactoryBase>(GetTransientPackage(), FactoryClass, NAME_None);
-	if (!FactoryToRegister)
-	{
-		return false;
-	}
-	if (FactoryToRegister->GetFactoryClass() == nullptr || RegisteredFactories.Contains(FactoryToRegister->GetFactoryClass()))
+	UClass* ClassToMake = FactoryClass->GetDefaultObject<UInterchangeFactoryBase>()->GetFactoryClass();
+	if (ClassToMake)
 	{
-		FactoryToRegister->MarkPendingKill();
-		return FactoryToRegister->GetFactoryClass() == nullptr ? false : true;
+		if (!RegisteredFactoryClasses.Contains(ClassToMake))
+		{
+			RegisteredFactoryClasses.Add(ClassToMake, FactoryClass);
+		}
+
+		return true;
 	}
-	RegisteredFactories.Add(FactoryToRegister->GetFactoryClass(), FactoryToRegister);
-	return true;
+
+	return false;
 }
 
 bool UInterchangeManager::RegisterWriter(const UClass* WriterClass)
@@ -416,7 +396,7 @@ bool UInterchangeManager::RegisterWriter(const UClass* WriterClass)
 		return false;
 	}
 
-	if (RegisteredFactories.Contains(WriterClass))
+	if (RegisteredWriters.Contains(WriterClass))
 	{
 		return true;
 	}
@@ -597,8 +577,9 @@ UE::Interchange::FAssetImportResultRef UInterchangeManager::ImportAssetAsync(con
 		check(AsyncHelper->BaseNodeContainers[SourceDataIndex].IsValid());
 	}
 
+	const UInterchangeProjectSettings* InterchangeProjectSettings = GetDefault<UInterchangeProjectSettings>();
 	UInterchangePipelineConfigurationBase* RegisteredPipelineConfiguration = nullptr;
-	TSoftClassPtr <UInterchangePipelineConfigurationBase> PipelineConfigurationDialogClass = GetDefault<UInterchangeProjectSettings>()->PipelineConfigurationDialogClass;
+	TSoftClassPtr <UInterchangePipelineConfigurationBase> PipelineConfigurationDialogClass = InterchangeProjectSettings->PipelineConfigurationDialogClass;
 	if (PipelineConfigurationDialogClass.IsValid())
 	{
 		UClass* PipelineConfigurationClass = PipelineConfigurationDialogClass.LoadSynchronous();
@@ -614,11 +595,11 @@ UE::Interchange::FAssetImportResultRef UInterchangeManager::ImportAssetAsync(con
 		const bool bIsUnattended = FApp::IsUnattended() || GIsAutomationTesting;
 
 		const bool bShowPipelineStacksConfigurationDialog = !bIsUnattended
-															&& GetDefault<UInterchangeProjectSettings>()->bShowPipelineStacksConfigurationDialog
+															&& InterchangeProjectSettings->bShowPipelineStacksConfigurationDialog
 															&& !bImportAllWithDefault;
 
-		const FName DefaultPipelineStackName = GetDefault<UInterchangeProjectSettings>()->DefaultPipelineStack;
-		const TMap<FName, FInterchangePipelineStack>& DefaultPipelineStacks = GetDefault<UInterchangeProjectSettings>()->PipelineStacks;
+		const FName DefaultPipelineStackName = InterchangeProjectSettings->DefaultPipelineStack;
+		const TMap<FName, FInterchangePipelineStack>& DefaultPipelineStacks = InterchangeProjectSettings->PipelineStacks;
 		
 		//If we reimport we want to load the original pipeline and the original pipeline settings
 		if (OriginalAssetImportData && OriginalAssetImportData->Pipelines.Num() > 0)
@@ -664,12 +645,12 @@ UE::Interchange::FAssetImportResultRef UInterchangeManager::ImportAssetAsync(con
 				{
 					bImportAllWithDefault = true;
 				}
-				PipelineStackName = GetDefault<UInterchangeProjectSettings>()->DefaultPipelineStack;
+				PipelineStackName = InterchangeProjectSettings->DefaultPipelineStack;
 			}
 			if (!bImportCancel)
 			{
 				//Get the latest PipelineStacks (the Dialog can change the CDO)
-				const TMap<FName, FInterchangePipelineStack>& PipelineStacks = GetDefault<UInterchangeProjectSettings>()->PipelineStacks;
+				const TMap<FName, FInterchangePipelineStack>& PipelineStacks = InterchangeProjectSettings->PipelineStacks;
 				if (!PipelineStacks.Contains(PipelineStackName))
 				{
 					if (PipelineStacks.Contains(DefaultPipelineStackName))
@@ -767,6 +748,26 @@ UInterchangeSourceData* UInterchangeManager::CreateSourceData(const FString& InF
 	return SourceDataAsset;
 }
 
+const UClass* UInterchangeManager::GetRegisteredFactoryClass(const UClass* ClassToMake) const
+{
+	const UClass* BestClassToMake = nullptr;
+	const UClass* Result = nullptr;
+
+	for (const auto& Kvp : RegisteredFactoryClasses)
+	{
+		if (ClassToMake->IsChildOf(Kvp.Key))
+		{
+			// Find the factory which handles the most derived registered type 
+			if (BestClassToMake == nullptr || Kvp.Key->IsChildOf(BestClassToMake))
+			{
+				BestClassToMake = Kvp.Key.Get();
+				Result = Kvp.Value.Get();
+			}
+		}
+	}
+	return Result;
+}
+
 TSharedPtr<UE::Interchange::FImportAsyncHelper, ESPMode::ThreadSafe> UInterchangeManager::CreateAsyncHelper(const UE::Interchange::FImportAsyncHelperData& Data)
 {
 	TSharedPtr<UE::Interchange::FImportAsyncHelper, ESPMode::ThreadSafe> AsyncHelper = MakeShared<UE::Interchange::FImportAsyncHelper, ESPMode::ThreadSafe>();
@@ -816,11 +817,12 @@ void UInterchangeManager::ReleaseAsyncHelper(TWeakPtr<UE::Interchange::FImportAs
 UInterchangeTranslatorBase* UInterchangeManager::GetTranslatorForSourceData(const UInterchangeSourceData* SourceData) const
 {
 	// Find the translator
-	for (UInterchangeTranslatorBase* Translator : RegisteredTranslators)
+	for (const UClass* TranslatorClass : RegisteredTranslatorsClass)
 	{
-		if (Translator->CanImportSourceData(SourceData))
+		if (TranslatorClass->GetDefaultObject<UInterchangeTranslatorBase>()->CanImportSourceData(SourceData))
 		{
-			UInterchangeTranslatorBase* SourceDataTranslator = NewObject<UInterchangeTranslatorBase>(GetTransientPackage(), Translator->GetClass(), NAME_None);
+			UInterchangeTranslatorBase* SourceDataTranslator = NewObject<UInterchangeTranslatorBase>(GetTransientPackage(), TranslatorClass, NAME_None);
+			SourceDataTranslator->SourceData = SourceData;
 			return SourceDataTranslator;
 		}
 	}
@@ -856,12 +858,13 @@ bool UInterchangeManager::CanTranslateSourceDataWithPayloadInterface(const UInte
 UInterchangeTranslatorBase* UInterchangeManager::GetTranslatorSupportingPayloadInterfaceForSourceData(const UInterchangeSourceData* SourceData, const UClass* PayloadInterfaceClass) const
 {
 	// Find the translator
-	for (UInterchangeTranslatorBase* Translator : RegisteredTranslators)
+	for (const UClass* TranslatorClass : RegisteredTranslatorsClass)
 	{
-		const UClass* TranslatorClass = Translator->GetClass();
-		if (TranslatorClass->ImplementsInterface(PayloadInterfaceClass) && Translator->CanImportSourceData(SourceData))
+		if (TranslatorClass->ImplementsInterface(PayloadInterfaceClass) &&
+			TranslatorClass->GetDefaultObject<UInterchangeTranslatorBase>()->CanImportSourceData(SourceData))
 		{
 			UInterchangeTranslatorBase* SourceDataTranslator = NewObject<UInterchangeTranslatorBase>(GetTransientPackage(), TranslatorClass, NAME_None);
+			SourceDataTranslator->SourceData = SourceData;
 			return SourceDataTranslator;
 		}
 	}
@@ -908,7 +911,7 @@ void UInterchangeManager::FindPipelineCandidate(TArray<UClass*>& PipelineCandida
 		PipelineCandidates.AddUnique(Class);
 	}
 
-//Blueprint and python script discoverability is available only if we compile with the engine
+	//Blueprint and python script discoverability is available only if we compile with the engine
 	// Load the asset registry module
 	FAssetRegistryModule& AssetRegistryModule = FModuleManager::LoadModuleChecked< FAssetRegistryModule >(FName("AssetRegistry"));
 	IAssetRegistry& AssetRegistry = AssetRegistryModule.Get();
diff --git a/Engine/Source/Runtime/Interchange/Engine/Private/Tasks/InterchangeTaskCompletion.cpp b/Engine/Source/Runtime/Interchange/Engine/Private/Tasks/InterchangeTaskCompletion.cpp
--- a/Engine/Source/Runtime/Interchange/Engine/Private/Tasks/InterchangeTaskCompletion.cpp
+++ b/Engine/Source/Runtime/Interchange/Engine/Private/Tasks/InterchangeTaskCompletion.cpp
@@ -7,6 +7,7 @@
 #include "InterchangeEngineLogPrivate.h"
 #include "InterchangeFactoryBase.h"
 #include "InterchangeManager.h"
+#include "InterchangeResultsContainer.h"
 #include "Stats/Stats.h"
 #include "Templates/SharedPointer.h"
 #include "UObject/Object.h"
@@ -36,6 +37,8 @@ void UE::Interchange::FTaskCompletion::DoTask(ENamedThreads::Type CurrentThread,
 	//No need anymore of the translators sources
 	AsyncHelper->ReleaseTranslatorsSource();
 
+	UInterchangeResultsContainer* Results = AsyncHelper->AssetImportResult->GetResults();
+
 	for(TPair<int32, TArray<FImportAsyncHelper::FImportedAssetInfo>>& AssetInfosPerSourceIndexPair : AsyncHelper->ImportedAssetsPerSourceIndex)
 	{
 		//Verify if the task was cancel
@@ -43,8 +46,10 @@ void UE::Interchange::FTaskCompletion::DoTask(ENamedThreads::Type CurrentThread,
 		{
 			break;
 		}
+
 		const int32 SourceIndex = AssetInfosPerSourceIndexPair.Key;
 		const bool bCallPostImportGameThreadCallback = ensure(AsyncHelper->SourceDatas.IsValidIndex(SourceIndex));
+
 		for (const FImportAsyncHelper::FImportedAssetInfo& AssetInfo : AssetInfosPerSourceIndexPair.Value)
 		{
 			UObject* Asset = AssetInfo.ImportAsset;
@@ -60,6 +65,14 @@ void UE::Interchange::FTaskCompletion::DoTask(ENamedThreads::Type CurrentThread,
 				AssetInfo.Factory->PostImportGameThreadCallback(Arguments);
 			}
 
+			if (Asset)
+			{
+				UInterchangeResultSuccess* Message = Results->Add<UInterchangeResultSuccess>();
+				Message->SourceAssetName = AsyncHelper->SourceDatas[SourceIndex]->GetFilename();
+				Message->DestinationAssetName = Asset->GetPathName();
+				Message->AssetType = Asset->GetClass();
+			}
+
 			//Clear any async flag from the created asset
 			const EInternalObjectFlags AsyncFlags = EInternalObjectFlags::Async | EInternalObjectFlags::AsyncLoading;
 			Asset->ClearInternalFlags(AsyncFlags);
