Removed test code which was left in by mistake.

diff --git a/Engine/Plugins/Runtime/Interchange/Source/Parsers/Fbx/Private/FbxMesh.cpp b/Engine/Plugins/Runtime/Interchange/Source/Parsers/Fbx/Private/FbxMesh.cpp
--- a/Engine/Plugins/Runtime/Interchange/Source/Parsers/Fbx/Private/FbxMesh.cpp
+++ b/Engine/Plugins/Runtime/Interchange/Source/Parsers/Fbx/Private/FbxMesh.cpp
@@ -295,9 +295,6 @@ namespace UE
 					return false;
 				}
 
-				UInterchangeResultMeshWarning_Generic* TestMessage = AddMessage<UInterchangeResultMeshWarning_Generic>(Mesh);
-				TestMessage->Text = LOCTEXT("TestMessage", "Test message for mesh '{MeshName}'.");
-
 				FStaticMeshAttributes Attributes(*MeshDescription);
 				Attributes.Register();
 
