In Interchange, add skeleton and physics asset import options to the generic pipeline. Those option allow the user to specify if we create or use an existing asset for skeleton and physics.

diff --git a/Engine/Plugins/Editor/InterchangeEditor/Source/Pipelines/Private/InterchangeEditorPipelineDetails.cpp b/Engine/Plugins/Editor/InterchangeEditor/Source/Pipelines/Private/InterchangeEditorPipelineDetails.cpp
--- a/Engine/Plugins/Editor/InterchangeEditor/Source/Pipelines/Private/InterchangeEditorPipelineDetails.cpp
+++ b/Engine/Plugins/Editor/InterchangeEditor/Source/Pipelines/Private/InterchangeEditorPipelineDetails.cpp
@@ -184,6 +184,12 @@ void FInterchangeBaseNodeDetailsCustomization::AddAttributeRow(UE::Interchange::
 		}
 		break;
 
+		case UE::Interchange::EAttributeTypes::SoftObjectPath:
+		{
+			BuildStringValueContent<FSoftObjectPath>(AttributeCategory, AttributeKey);
+		}
+		break;
+
 		default:
 		{
 			FText AttributeName = FText::FromString(InterchangeBaseNode->GetKeyDisplayName(AttributeKey));
@@ -356,6 +362,16 @@ void FInterchangeBaseNodeDetailsCustomization::BuildStringValueContent(IDetailCa
 					}
 					ReturnText = FText::FromName(Value);
 				}
+				else if (AttributeType == UE::Interchange::EAttributeTypes::SoftObjectPath)
+				{
+					UE::Interchange::FAttributeStorage::TAttributeHandle<FSoftObjectPath> AttributeHandle = InterchangeBaseNode->GetAttributeHandle<FSoftObjectPath>(AttributeKey);
+					FSoftObjectPath Value;
+					if (AttributeHandle.IsValid())
+					{
+						AttributeHandle.Get(Value);
+					}
+					ReturnText = FText::FromString(Value.ToString());
+				}
                 return ReturnText;
             })
             .Font(IDetailLayoutBuilder::GetDetailFont())
diff --git a/Engine/Plugins/Editor/InterchangeEditor/Source/Pipelines/Private/SInterchangePipelineConfigurationDialog.cpp b/Engine/Plugins/Editor/InterchangeEditor/Source/Pipelines/Private/SInterchangePipelineConfigurationDialog.cpp
--- a/Engine/Plugins/Editor/InterchangeEditor/Source/Pipelines/Private/SInterchangePipelineConfigurationDialog.cpp
+++ b/Engine/Plugins/Editor/InterchangeEditor/Source/Pipelines/Private/SInterchangePipelineConfigurationDialog.cpp
@@ -66,6 +66,7 @@ void SInterchangePipelineStacksTreeView::Construct(const FArguments& InArgs)
 				UInterchangePipelineBase* GeneratedPipeline = NewObject<UInterchangePipelineBase>(GetTransientPackage(), PipelineClass, NAME_None, RF_NoFlags);
 				//Load the settings for this pipeline
 				GeneratedPipeline->LoadSettings(NameAndPipelineStack.Key);
+				GeneratedPipeline->PreDialogCleanup(NameAndPipelineStack.Key);
 				TSharedPtr<FInterchangePipelineStacksTreeNodeItem> PipelineNode = MakeShared<FInterchangePipelineStacksTreeNodeItem>();
 				PipelineNode->StackName = NameAndPipelineStack.Key;
 				PipelineNode->Pipeline = GeneratedPipeline;
diff --git a/Engine/Plugins/Runtime/Interchange/Source/Import/InterchangeImport.Build.cs b/Engine/Plugins/Runtime/Interchange/Source/Import/InterchangeImport.Build.cs
--- a/Engine/Plugins/Runtime/Interchange/Source/Import/InterchangeImport.Build.cs
+++ b/Engine/Plugins/Runtime/Interchange/Source/Import/InterchangeImport.Build.cs
@@ -24,6 +24,7 @@ namespace UnrealBuildTool.Rules
 			PrivateDependencyModuleNames.AddRange(
 				new string[]
 				{
+					"AssetRegistry",
 					"IESFile",
 					"ImageWrapper",
 					"InterchangeDispatcher",
diff --git a/Engine/Plugins/Runtime/Interchange/Source/Import/Private/InterchangeImportModule.cpp b/Engine/Plugins/Runtime/Interchange/Source/Import/Private/InterchangeImportModule.cpp
--- a/Engine/Plugins/Runtime/Interchange/Source/Import/Private/InterchangeImportModule.cpp
+++ b/Engine/Plugins/Runtime/Interchange/Source/Import/Private/InterchangeImportModule.cpp
@@ -7,6 +7,7 @@
 #include "InterchangeImportLog.h"
 #include "InterchangeManager.h"
 #include "Material/InterchangeMaterialFactory.h"
+#include "Mesh/InterchangePhysicsAssetFactory.h"
 #include "Mesh/InterchangeSkeletalMeshFactory.h"
 #include "Mesh/InterchangeSkeletonFactory.h"
 #include "Mesh/InterchangeStaticMeshFactory.h"
@@ -70,6 +71,7 @@ void FInterchangeImportModule::StartupModule()
 		InterchangeManager.RegisterFactory(UInterchangeSkeletonFactory::StaticClass());
 		InterchangeManager.RegisterFactory(UInterchangeSkeletalMeshFactory::StaticClass());
 		InterchangeManager.RegisterFactory(UInterchangeStaticMeshFactory::StaticClass());
+		InterchangeManager.RegisterFactory(UInterchangePhysicsAssetFactory::StaticClass());
 	};
 
 	if (GEngine)
diff --git a/Engine/Plugins/Runtime/Interchange/Source/Import/Private/Material/InterchangeMaterialFactory.cpp b/Engine/Plugins/Runtime/Interchange/Source/Import/Private/Material/InterchangeMaterialFactory.cpp
--- a/Engine/Plugins/Runtime/Interchange/Source/Import/Private/Material/InterchangeMaterialFactory.cpp
+++ b/Engine/Plugins/Runtime/Interchange/Source/Import/Private/Material/InterchangeMaterialFactory.cpp
@@ -302,10 +302,10 @@ UObject* UInterchangeMaterialFactory::CreateAsset(const UInterchangeMaterialFact
 }
 
 /* This function is call in the completion task on the main thread, use it to call main thread post creation step for your assets*/
-void UInterchangeMaterialFactory::PostImportGameThreadCallback(const FPostImportGameThreadCallbackParams& Arguments) const
+void UInterchangeMaterialFactory::PreImportPreCompletedCallback(const FImportPreCompletedCallbackParams& Arguments) const
 {
 	check(IsInGameThread());
-	Super::PostImportGameThreadCallback(Arguments);
+	Super::PreImportPreCompletedCallback(Arguments);
 
 	//TODO make sure this work at runtime
 #if WITH_EDITORONLY_DATA
diff --git a/Engine/Plugins/Runtime/Interchange/Source/Import/Private/Mesh/InterchangePhysicsAssetFactory.cpp b/Engine/Plugins/Runtime/Interchange/Source/Import/Private/Mesh/InterchangePhysicsAssetFactory.cpp
new file mode 100644
--- /dev/null
+++ b/Engine/Plugins/Runtime/Interchange/Source/Import/Private/Mesh/InterchangePhysicsAssetFactory.cpp
@@ -0,0 +1,135 @@
+// Copyright Epic Games, Inc. All Rights Reserved. 
+#include "Mesh/InterchangePhysicsAssetFactory.h"
+
+#include "InterchangeImportCommon.h"
+#include "InterchangeImportLog.h"
+#include "InterchangePhysicsAssetFactoryNode.h"
+#include "InterchangeSourceData.h"
+#include "Nodes/InterchangeBaseNode.h"
+#include "Nodes/InterchangeBaseNodeContainer.h"
+#include "PhysicsEngine/PhysicsAsset.h"
+
+UClass* UInterchangePhysicsAssetFactory::GetFactoryClass() const
+{
+	return UPhysicsAsset::StaticClass();
+}
+
+UObject* UInterchangePhysicsAssetFactory::CreateEmptyAsset(const FCreateAssetParams& Arguments) const
+{
+	UObject* PhysicsAsset = nullptr;
+
+#if WITH_EDITORONLY_DATA
+	if (!Arguments.AssetNode || !Arguments.AssetNode->GetAssetClass()->IsChildOf(GetFactoryClass()))
+	{
+		return nullptr;
+	}
+
+	const UInterchangePhysicsAssetFactoryNode* PhysicsAssetNode = Cast<UInterchangePhysicsAssetFactoryNode>(Arguments.AssetNode);
+	if (PhysicsAssetNode == nullptr)
+	{
+		return nullptr;
+	}
+
+	// create an asset if it doesn't exist
+	UObject* ExistingAsset = StaticFindObject(nullptr, Arguments.Parent, *Arguments.AssetName);
+
+	// create a new PhysicsAsset or overwrite existing asset, if possible
+	if (!ExistingAsset)
+	{
+		PhysicsAsset = NewObject<UObject>(Arguments.Parent, UPhysicsAsset::StaticClass(), *Arguments.AssetName, RF_Public | RF_Standalone);
+	}
+	else if (ExistingAsset->GetClass()->IsChildOf(UPhysicsAsset::StaticClass()))
+	{
+		//This is a reimport, we are just re-updating the source data
+		PhysicsAsset = ExistingAsset;
+	}
+
+	if (!PhysicsAsset)
+	{
+		UE_LOG(LogInterchangeImport, Warning, TEXT("Could not create PhysicsAsset asset %s"), *Arguments.AssetName);
+		return nullptr;
+	}
+	PhysicsAssetNode->ReferenceObject = PhysicsAsset;
+
+	PhysicsAsset->PreEditChange(nullptr);
+#endif //WITH_EDITORONLY_DATA
+	return PhysicsAsset;
+}
+
+UObject* UInterchangePhysicsAssetFactory::CreateAsset(const UInterchangePhysicsAssetFactory::FCreateAssetParams& Arguments) const
+{
+#if !WITH_EDITORONLY_DATA
+
+	UE_LOG(LogInterchangeImport, Error, TEXT("Cannot import PhysicsAsset asset in runtime, this is an editor only feature."));
+	return nullptr;
+
+#else
+
+	if (!Arguments.AssetNode || !Arguments.AssetNode->GetAssetClass()->IsChildOf(GetFactoryClass()))
+	{
+		return nullptr;
+	}
+
+	const UInterchangePhysicsAssetFactoryNode* PhysicsAssetNode = Cast<UInterchangePhysicsAssetFactoryNode>(Arguments.AssetNode);
+	if (PhysicsAssetNode == nullptr)
+	{
+		return nullptr;
+	}
+
+	const UClass* PhysicsAssetClass = PhysicsAssetNode->GetAssetClass();
+	check(PhysicsAssetClass && PhysicsAssetClass->IsChildOf(GetFactoryClass()));
+
+	// create an asset if it doesn't exist
+	UObject* ExistingAsset = StaticFindObject(nullptr, Arguments.Parent, *Arguments.AssetName);
+
+	UObject* PhysicsAssetObject = nullptr;
+	// create a new PhysicsAsset or overwrite existing asset, if possible
+	if (!ExistingAsset)
+	{
+		//NewObject is not thread safe, the asset registry directory watcher tick on the main thread can trig before we finish initializing the UObject and will crash
+		//The UObject should have been create by calling CreateEmptyAsset on the main thread.
+		check(IsInGameThread());
+		PhysicsAssetObject = NewObject<UObject>(Arguments.Parent, PhysicsAssetClass, *Arguments.AssetName, RF_Public | RF_Standalone);
+	}
+	else if(ExistingAsset->GetClass()->IsChildOf(PhysicsAssetClass))
+	{
+		//This is a reimport, we are just re-updating the source data
+		PhysicsAssetObject = ExistingAsset;
+	}
+
+	if (!PhysicsAssetObject)
+	{
+		UE_LOG(LogInterchangeImport, Warning, TEXT("Could not create PhysicsAsset asset %s"), *Arguments.AssetName);
+		return nullptr;
+	}
+
+	if (PhysicsAssetObject)
+	{
+		//Currently PhysicsAsset re-import will not touch the PhysicsAsset at all
+		//TODO design a re-import process for the PhysicsAsset
+		if(!Arguments.ReimportObject)
+		{
+			UPhysicsAsset* PhysicsAsset = Cast<UPhysicsAsset>(PhysicsAssetObject);
+			if (!ensure(PhysicsAsset))
+			{
+				UE_LOG(LogInterchangeImport, Warning, TEXT("Could not create PhysicsAsset asset %s"), *Arguments.AssetName);
+				return nullptr;
+			}
+			
+			PhysicsAssetNode->ReferenceObject = PhysicsAsset;
+		}
+		
+		//Getting the file Hash will cache it into the source data
+		Arguments.SourceData->GetFileContentHash();
+
+		//The interchange completion task (call in the GameThread after the factories pass), will call PostEditChange.
+	}
+	else
+	{
+		//The PhysicsAsset is not a UPhysicsAsset
+		PhysicsAssetObject->RemoveFromRoot();
+		PhysicsAssetObject->MarkPendingKill();
+	}
+	return PhysicsAssetObject;
+#endif
+}
diff --git a/Engine/Plugins/Runtime/Interchange/Source/Import/Private/Mesh/InterchangeSkeletalMeshFactory.cpp b/Engine/Plugins/Runtime/Interchange/Source/Import/Private/Mesh/InterchangeSkeletalMeshFactory.cpp
--- a/Engine/Plugins/Runtime/Interchange/Source/Import/Private/Mesh/InterchangeSkeletalMeshFactory.cpp
+++ b/Engine/Plugins/Runtime/Interchange/Source/Import/Private/Mesh/InterchangeSkeletalMeshFactory.cpp
@@ -215,7 +215,7 @@ namespace UE
 				}
 			}
 
-			bool CopyBlendShapesMeshDescriptionToSkeletalMeshImportData(const TMap<FString, TOptional<UE::Interchange::FSkeletalMeshBlendShapePayloadData>>& LodBlendShapeMeshDescriptions, FSkeletalMeshImportData& DestinationSkeletalMeshImportData)
+			void CopyBlendShapesMeshDescriptionToSkeletalMeshImportData(const TMap<FString, TOptional<UE::Interchange::FSkeletalMeshBlendShapePayloadData>>& LodBlendShapeMeshDescriptions, FSkeletalMeshImportData& DestinationSkeletalMeshImportData)
 			{
 				const int32 OriginalMorphTargetCount = LodBlendShapeMeshDescriptions.Num();
 				TArray<FString> Keys;
@@ -242,6 +242,13 @@ namespace UE
 					Keys.Add(Pair.Key);
 					MorphTargetCount++;
 				}
+
+				//No morph target to import
+				if (MorphTargetCount == 0)
+				{
+					return;
+				}
+
 				ensure(Keys.Num() == MorphTargetCount);
 				//Allocate the data
 				DestinationSkeletalMeshImportData.MorphTargetNames.AddDefaulted(MorphTargetCount);
@@ -250,6 +257,9 @@ namespace UE
 
 				int32 NumMorphGroup = FMath::Min(FPlatformMisc::NumberOfWorkerThreadsToSpawn(), MorphTargetCount);
 				const int32 MorphTargetGroupSize = FMath::Max(FMath::CeilToInt(static_cast<float>(MorphTargetCount) / static_cast<float>(NumMorphGroup)), 1);
+				//Re-Adjust the group Number in case we have a reminder error (exemple MorphTargetGroupSize = 4.8 -> 5 so the number of group can be lower if there is a large amount of Group)
+				NumMorphGroup = FMath::CeilToInt(static_cast<float>(MorphTargetCount) / static_cast<float>(MorphTargetGroupSize));
+
 				ParallelFor(NumMorphGroup, [MorphTargetGroupSize,
 							MorphTargetCount,
 							NumMorphGroup,
@@ -264,7 +274,7 @@ namespace UE
 						if (!Keys.IsValidIndex(MorphTargetIndex))
 						{
 							ensure(MorphTargetGroupIndex + 1 == NumMorphGroup);
-							//Executing the last morph target group, and we do not have a full group.
+							//Executing the last morph target group, in case we do not have a full last group.
 							break;
 						}
 						const FString BlendShapeKey(Keys[MorphTargetIndex]);
@@ -319,10 +329,11 @@ namespace UE
 					}
 				}
 				, EParallelForFlags::BackgroundPriority);
-				return true;
+				return;
 			}
 
-			void RetrieveAllSkeletalMeshPayloadsAndFillImportData(FSkeletalMeshImportData& DestinationImportData
+			void RetrieveAllSkeletalMeshPayloadsAndFillImportData(const UInterchangeSkeletalMeshFactoryNode* SkeletalMeshFactoryNode
+																  , FSkeletalMeshImportData& DestinationImportData
 																  , TArray<FMeshNodeContext>& MeshReferences
 																  , TArray<SkeletalMeshImportData::FBone>& RefBonesBinary
 																  , const UInterchangeSkeletalMeshFactory::FCreateAssetParams& Arguments
@@ -341,6 +352,9 @@ namespace UE
 					AppendSettings.bMergeUVChannels[ChannelIdx] = true;
 				}
 
+				bool bImportMorphTarget = true;
+				SkeletalMeshFactoryNode->GetCustomImportMorphTarget(bImportMorphTarget);
+
 				TMap<FString, TFuture<TOptional<UE::Interchange::FSkeletalMeshLodPayloadData>>> LodMeshPayloadPerTranslatorPayloadKey;
 				LodMeshPayloadPerTranslatorPayloadKey.Reserve(MeshReferences.Num());
 
@@ -352,7 +366,7 @@ namespace UE
 					//Add the payload entry key, the payload data will be fill later in bulk by the translator
 					LodMeshPayloadPerTranslatorPayloadKey.Add(MeshNodeContext.TranslatorPayloadKey, SkeletalMeshTranslatorPayloadInterface->GetSkeletalMeshLodPayloadData(MeshNodeContext.TranslatorPayloadKey));
 					//Count the blend shape dependencies so we can reserve the right amount
-					BlendShapeCount += MeshNodeContext.MeshNode ? MeshNodeContext.MeshNode->GetShapeDependeciesCount() : 0;
+					BlendShapeCount += (bImportMorphTarget && MeshNodeContext.MeshNode) ? MeshNodeContext.MeshNode->GetShapeDependeciesCount() : 0;
 				}
 				BlendShapeMeshDescriptionsPerBlendShapeName.Reserve(BlendShapeCount);
 
@@ -404,13 +418,15 @@ namespace UE
 					}
 					FStaticMeshOperations::AppendMeshDescription(LodMeshPayload->LodMeshDescription, LodMeshDescription, AppendSettings);
 					FSkeletalMeshOperations::AppendSkinWeight(LodMeshPayload->LodMeshDescription, LodMeshDescription, SkeletalMeshAppendSettings);
-
-					FillBlendShapeMeshDescriptionsPerBlendShapeName(MeshNodeContext
-																	, BlendShapeMeshDescriptionsPerBlendShapeName
-																	, SkeletalMeshTranslatorPayloadInterface
-																	, VertexOffset
-																	, Arguments.NodeContainer
-																	, Arguments.AssetName);
+					if (bImportMorphTarget)
+					{
+						FillBlendShapeMeshDescriptionsPerBlendShapeName(MeshNodeContext
+																		, BlendShapeMeshDescriptionsPerBlendShapeName
+																		, SkeletalMeshTranslatorPayloadInterface
+																		, VertexOffset
+																		, Arguments.NodeContainer
+																		, Arguments.AssetName);
+					}
 				}
 
 				DestinationImportData = FSkeletalMeshImportData::CreateFromMeshDescription(LodMeshDescription);
@@ -640,8 +656,6 @@ namespace UE
 					}
 				}
 			}
-
-
 		} //Namespace Private
 	} //namespace Interchange
 } //namespace UE
@@ -854,12 +868,26 @@ UObject* UInterchangeSkeletalMeshFactory::CreateAsset(const UInterchangeSkeletal
 					continue;
 				}
 				
+				FSoftObjectPath SpecifiedSkeleton;
+				SkeletalMeshFactoryNode->GetCustomSkeletonSoftObjectPath(SpecifiedSkeleton);
+				bool bSpecifiedSkeleton = SpecifiedSkeleton.IsValid();
 				if(SkeletonReference == nullptr)
 				{
-					UObject* SkeletonObject = SkeletonNode->ReferenceObject.ResolveObject();
+					UObject* SkeletonObject = nullptr;
+					
+					if (SpecifiedSkeleton.IsValid())
+					{
+						SkeletonObject = SpecifiedSkeleton.TryLoad();
+					}
+					else if (SkeletonNode->ReferenceObject.IsValid())
+					{
+						SkeletonObject = SkeletonNode->ReferenceObject.TryLoad();
+					}
+
 					if (SkeletonObject)
 					{
 						SkeletonReference = Cast<USkeleton>(SkeletonObject);
+						
 					}
 					if (!ensure(SkeletonReference))
 					{
@@ -878,12 +906,16 @@ UObject* UInterchangeSkeletalMeshFactory::CreateAsset(const UInterchangeSkeletal
 				int32 SkeletonDepth = 0;
 				TArray<SkeletalMeshImportData::FBone> RefBonesBinary;
 				UE::Interchange::Private::ProcessImportMeshSkeleton(SkeletonReference, SkeletalMesh->GetRefSkeleton(), SkeletonDepth, Arguments.NodeContainer, RootJointNodeId, RefBonesBinary);
-
+				if (bSpecifiedSkeleton && !SkeletonReference->IsCompatibleMesh(SkeletalMesh))
+				{
+					UE_LOG(LogInterchangeImport, Warning, TEXT("The skeleton %s is incompatible with the imported skeletalmesh asset %s"), *SkeletonReference->GetName() , *Arguments.AssetName);
+				}
 				//Add the lod mesh data to the skeletalmesh
 				FSkeletalMeshImportData SkeletalMeshImportData;
 				
 				//Get all meshes and blend shapes payload and fill the SkeletalMeshImportData structure
-				UE::Interchange::Private::RetrieveAllSkeletalMeshPayloadsAndFillImportData(SkeletalMeshImportData
+				UE::Interchange::Private::RetrieveAllSkeletalMeshPayloadsAndFillImportData(SkeletalMeshFactoryNode
+																						   , SkeletalMeshImportData
 																						   , MeshReferences
 																						   , RefBonesBinary
 																						   , Arguments
@@ -950,6 +982,20 @@ UObject* UInterchangeSkeletalMeshFactory::CreateAsset(const UInterchangeSkeletal
 
 			/** Apply all SkeletalMeshFactoryNode custom attributes to the material asset */
 			SkeletalMeshFactoryNode->ApplyAllCustomAttributeToAsset(SkeletalMesh);
+
+			bool bCreatePhysicsAsset = false;
+			SkeletalMeshFactoryNode->GetCustomCreatePhysicsAsset(bCreatePhysicsAsset);
+			
+			if (!bCreatePhysicsAsset)
+			{
+				FSoftObjectPath SpecifiedPhysicAsset;
+				SkeletalMeshFactoryNode->GetCustomPhysicAssetSoftObjectPath(SpecifiedPhysicAsset);
+				if (SpecifiedPhysicAsset.IsValid())
+				{
+					UPhysicsAsset* PhysicsAsset = Cast<UPhysicsAsset>(SpecifiedPhysicAsset.TryLoad());
+					SkeletalMesh->SetPhysicsAsset(PhysicsAsset);
+				}
+			}
 		}
 		
 		//Getting the file Hash will cache it into the source data
@@ -969,10 +1015,10 @@ UObject* UInterchangeSkeletalMeshFactory::CreateAsset(const UInterchangeSkeletal
 }
 
 /* This function is call in the completion task on the main thread, use it to call main thread post creation step for your assets*/
-void UInterchangeSkeletalMeshFactory::PostImportGameThreadCallback(const FPostImportGameThreadCallbackParams& Arguments) const
+void UInterchangeSkeletalMeshFactory::PreImportPreCompletedCallback(const FImportPreCompletedCallbackParams& Arguments) const
 {
 	check(IsInGameThread());
-	Super::PostImportGameThreadCallback(Arguments);
+	Super::PreImportPreCompletedCallback(Arguments);
 
 	//TODO make sure this work at runtime
 #if WITH_EDITORONLY_DATA
@@ -993,4 +1039,4 @@ void UInterchangeSkeletalMeshFactory::PostImportGameThreadCallback(const FPostIm
 		SkeletalMesh->SetAssetImportData(ImportDataPtr);
 	}
 #endif
-}
\ No newline at end of file
+}
diff --git a/Engine/Plugins/Runtime/Interchange/Source/Import/Private/Mesh/InterchangeStaticMeshFactory.cpp b/Engine/Plugins/Runtime/Interchange/Source/Import/Private/Mesh/InterchangeStaticMeshFactory.cpp
--- a/Engine/Plugins/Runtime/Interchange/Source/Import/Private/Mesh/InterchangeStaticMeshFactory.cpp
+++ b/Engine/Plugins/Runtime/Interchange/Source/Import/Private/Mesh/InterchangeStaticMeshFactory.cpp
@@ -349,10 +349,10 @@ UObject* UInterchangeStaticMeshFactory::CreateAsset(const UInterchangeStaticMesh
 
 
 /* This function is call in the completion task on the main thread, use it to call main thread post creation step for your assets */
-void UInterchangeStaticMeshFactory::PostImportGameThreadCallback(const FPostImportGameThreadCallbackParams& Arguments) const
+void UInterchangeStaticMeshFactory::PreImportPreCompletedCallback(const FImportPreCompletedCallbackParams& Arguments) const
 {
 	check(IsInGameThread());
-	Super::PostImportGameThreadCallback(Arguments);
+	Super::PreImportPreCompletedCallback(Arguments);
 
 	// TODO: make sure this works at runtime
 #if WITH_EDITORONLY_DATA
diff --git a/Engine/Plugins/Runtime/Interchange/Source/Import/Private/Texture/InterchangeTextureFactory.cpp b/Engine/Plugins/Runtime/Interchange/Source/Import/Private/Texture/InterchangeTextureFactory.cpp
--- a/Engine/Plugins/Runtime/Interchange/Source/Import/Private/Texture/InterchangeTextureFactory.cpp
+++ b/Engine/Plugins/Runtime/Interchange/Source/Import/Private/Texture/InterchangeTextureFactory.cpp
@@ -725,10 +725,10 @@ UObject* UInterchangeTextureFactory::CreateAsset(const UInterchangeTextureFactor
 }
 
 /* This function is call in the completion task on the main thread, use it to call main thread post creation step for your assets*/
-void UInterchangeTextureFactory::PostImportGameThreadCallback(const FPostImportGameThreadCallbackParams& Arguments) const
+void UInterchangeTextureFactory::PreImportPreCompletedCallback(const FImportPreCompletedCallbackParams& Arguments) const
 {
 	check(IsInGameThread());
-	Super::PostImportGameThreadCallback(Arguments);
+	Super::PreImportPreCompletedCallback(Arguments);
 	//TODO make sure this work at runtime
 #if WITH_EDITORONLY_DATA
 	if (ensure(Arguments.ImportedObject && Arguments.SourceData))
diff --git a/Engine/Plugins/Runtime/Interchange/Source/Import/Public/Material/InterchangeMaterialFactory.h b/Engine/Plugins/Runtime/Interchange/Source/Import/Public/Material/InterchangeMaterialFactory.h
--- a/Engine/Plugins/Runtime/Interchange/Source/Import/Public/Material/InterchangeMaterialFactory.h
+++ b/Engine/Plugins/Runtime/Interchange/Source/Import/Public/Material/InterchangeMaterialFactory.h
@@ -21,7 +21,7 @@ public:
 	virtual UClass* GetFactoryClass() const override;
 	virtual UObject* CreateEmptyAsset(const FCreateAssetParams& Arguments) const override;
 	virtual UObject* CreateAsset(const FCreateAssetParams& Arguments) const override;
-	virtual void PostImportGameThreadCallback(const FPostImportGameThreadCallbackParams& Arguments) const override;
+	virtual void PreImportPreCompletedCallback(const FImportPreCompletedCallbackParams& Arguments) const override;
 
 	// Interchange factory base interface end
 	//////////////////////////////////////////////////////////////////////////
diff --git a/Engine/Plugins/Runtime/Interchange/Source/Import/Public/Mesh/InterchangePhysicsAssetFactory.h b/Engine/Plugins/Runtime/Interchange/Source/Import/Public/Mesh/InterchangePhysicsAssetFactory.h
new file mode 100644
--- /dev/null
+++ b/Engine/Plugins/Runtime/Interchange/Source/Import/Public/Mesh/InterchangePhysicsAssetFactory.h
@@ -0,0 +1,31 @@
+// Copyright Epic Games, Inc. All Rights Reserved.
+
+#pragma once
+
+#include "CoreMinimal.h"
+#include "InterchangeFactoryBase.h"
+#include "UObject/Object.h"
+#include "UObject/ObjectMacros.h"
+
+#include "InterchangePhysicsAssetFactory.generated.h"
+
+UCLASS(BlueprintType)
+class INTERCHANGEIMPORT_API UInterchangePhysicsAssetFactory : public UInterchangeFactoryBase
+{
+	GENERATED_BODY()
+public:
+
+	//////////////////////////////////////////////////////////////////////////
+	// Interchange factory base interface begin
+
+	virtual UClass* GetFactoryClass() const override;
+	virtual UObject* CreateEmptyAsset(const FCreateAssetParams& Arguments) const override;
+	virtual UObject* CreateAsset(const FCreateAssetParams& Arguments) const override;
+	//virtual void PreImportPreCompletedCallback(const FImportPreCompletedCallbackParams& Arguments) const override;
+	//virtual void PostImportPreCompletedCallback(const FImportPreCompletedCallbackParams& Arguments) const override;
+
+	// Interchange factory base interface end
+	//////////////////////////////////////////////////////////////////////////
+};
+
+
diff --git a/Engine/Plugins/Runtime/Interchange/Source/Import/Public/Mesh/InterchangeSkeletalMeshFactory.h b/Engine/Plugins/Runtime/Interchange/Source/Import/Public/Mesh/InterchangeSkeletalMeshFactory.h
--- a/Engine/Plugins/Runtime/Interchange/Source/Import/Public/Mesh/InterchangeSkeletalMeshFactory.h
+++ b/Engine/Plugins/Runtime/Interchange/Source/Import/Public/Mesh/InterchangeSkeletalMeshFactory.h
@@ -9,6 +9,8 @@
 
 #include "InterchangeSkeletalMeshFactory.generated.h"
 
+class USkeletalMesh;
+
 UCLASS(BlueprintType)
 class INTERCHANGEIMPORT_API UInterchangeSkeletalMeshFactory : public UInterchangeFactoryBase
 {
@@ -21,7 +23,8 @@ public:
 	virtual UClass* GetFactoryClass() const override;
 	virtual UObject* CreateEmptyAsset(const FCreateAssetParams& Arguments) const override;
 	virtual UObject* CreateAsset(const FCreateAssetParams& Arguments) const override;
-	virtual void PostImportGameThreadCallback(const FPostImportGameThreadCallbackParams& Arguments) const override;
+	virtual void PreImportPreCompletedCallback(const FImportPreCompletedCallbackParams& Arguments) const override;
+	//virtual void PostImportPreCompletedCallback(const FImportPreCompletedCallbackParams& Arguments) const override;
 
 	// Interchange factory base interface end
 	//////////////////////////////////////////////////////////////////////////
diff --git a/Engine/Plugins/Runtime/Interchange/Source/Import/Public/Mesh/InterchangeSkeletonFactory.h b/Engine/Plugins/Runtime/Interchange/Source/Import/Public/Mesh/InterchangeSkeletonFactory.h
--- a/Engine/Plugins/Runtime/Interchange/Source/Import/Public/Mesh/InterchangeSkeletonFactory.h
+++ b/Engine/Plugins/Runtime/Interchange/Source/Import/Public/Mesh/InterchangeSkeletonFactory.h
@@ -21,7 +21,7 @@ public:
 	virtual UClass* GetFactoryClass() const override;
 	virtual UObject* CreateEmptyAsset(const FCreateAssetParams& Arguments) const override;
 	virtual UObject* CreateAsset(const FCreateAssetParams& Arguments) const override;
-	//virtual void PostImportGameThreadCallback(const FPostImportGameThreadCallbackParams& Arguments) const override;
+	//virtual void PreImportPreCompletedCallback(const FImportPreCompletedCallbackParams& Arguments) const override;
 
 	// Interchange factory base interface end
 	//////////////////////////////////////////////////////////////////////////
diff --git a/Engine/Plugins/Runtime/Interchange/Source/Import/Public/Mesh/InterchangeStaticMeshFactory.h b/Engine/Plugins/Runtime/Interchange/Source/Import/Public/Mesh/InterchangeStaticMeshFactory.h
--- a/Engine/Plugins/Runtime/Interchange/Source/Import/Public/Mesh/InterchangeStaticMeshFactory.h
+++ b/Engine/Plugins/Runtime/Interchange/Source/Import/Public/Mesh/InterchangeStaticMeshFactory.h
@@ -21,7 +21,8 @@ public:
 	virtual UClass* GetFactoryClass() const override;
 	virtual UObject* CreateEmptyAsset(const FCreateAssetParams& Arguments) const override;
 	virtual UObject* CreateAsset(const FCreateAssetParams& Arguments) const override;
-	virtual void PostImportGameThreadCallback(const FPostImportGameThreadCallbackParams& Arguments) const override;
+	virtual void PreImportPreCompletedCallback(const FImportPreCompletedCallbackParams& Arguments) const override;
+	//virtual void PostImportPreCompletedCallback(const FImportPreCompletedCallbackParams& Arguments) const override;
 
 	// Interchange factory base interface end
 	//////////////////////////////////////////////////////////////////////////
diff --git a/Engine/Plugins/Runtime/Interchange/Source/Import/Public/Texture/InterchangeTextureFactory.h b/Engine/Plugins/Runtime/Interchange/Source/Import/Public/Texture/InterchangeTextureFactory.h
--- a/Engine/Plugins/Runtime/Interchange/Source/Import/Public/Texture/InterchangeTextureFactory.h
+++ b/Engine/Plugins/Runtime/Interchange/Source/Import/Public/Texture/InterchangeTextureFactory.h
@@ -22,7 +22,7 @@ public:
 	virtual UClass* GetFactoryClass() const override;
 	virtual UObject* CreateEmptyAsset(const FCreateAssetParams& Arguments) const override;
 	virtual UObject* CreateAsset(const FCreateAssetParams& Arguments) const override;
-	virtual void PostImportGameThreadCallback(const FPostImportGameThreadCallbackParams& Arguments) const override;
+	virtual void PreImportPreCompletedCallback(const FImportPreCompletedCallbackParams& Arguments) const override;
 	// Interchange factory base interface end
 	//////////////////////////////////////////////////////////////////////////
 };
diff --git a/Engine/Plugins/Runtime/Interchange/Source/Nodes/Public/InterchangePhysicsAssetFactoryNode.h b/Engine/Plugins/Runtime/Interchange/Source/Nodes/Public/InterchangePhysicsAssetFactoryNode.h
new file mode 100644
--- /dev/null
+++ b/Engine/Plugins/Runtime/Interchange/Source/Nodes/Public/InterchangePhysicsAssetFactoryNode.h
@@ -0,0 +1,140 @@
+// Copyright Epic Games, Inc. All Rights Reserved.
+
+#pragma once
+
+#include "CoreMinimal.h"
+#include "Nodes/InterchangeBaseNode.h"
+
+#if WITH_ENGINE
+#include "PhysicsEngine/PhysicsAsset.h"
+#endif
+
+#include "InterchangePhysicsAssetFactoryNode.generated.h"
+
+UCLASS(BlueprintType)
+class INTERCHANGENODES_API UInterchangePhysicsAssetFactoryNode : public UInterchangeBaseNode
+{
+	GENERATED_BODY()
+
+public:
+	UInterchangePhysicsAssetFactoryNode()
+	:UInterchangeBaseNode()
+	{
+#if WITH_ENGINE
+		AssetClass = nullptr;
+#endif
+	}
+
+	/**
+	 * Initialize node data
+	 * @param: UniqueID - The uniqueId for this node
+	 * @param DisplayLabel - The name of the node
+	 * @param InAssetClass - The class the Skeleton factory will create for this node.
+	 */
+	UFUNCTION(BlueprintCallable, Category = "Interchange | Node | Skeleton")
+	void InitializePhysicsAssetNode(const FString& UniqueID, const FString& DisplayLabel, const FString& InAssetClass)
+	{
+		bIsNodeClassInitialized = false;
+		InitializeNode(UniqueID, DisplayLabel, EInterchangeNodeContainerType::NodeContainerType_FactoryData);
+		FString OperationName = GetTypeName() + TEXT(".SetAssetClassName");
+		InterchangePrivateNodeBase::SetCustomAttribute<FString>(*Attributes, ClassNameAttributeKey, OperationName, InAssetClass);
+		FillAssetClassFromAttribute();
+	}
+
+	virtual void Serialize(FArchive& Ar) override
+	{
+		Super::Serialize(Ar);
+#if WITH_ENGINE
+		if(Ar.IsLoading())
+		{
+			//Make sure the class is properly set when we compile with engine, this will set the
+			//bIsNodeClassInitialized to true.
+			SetNodeClassFromClassAttribute();
+		}
+#endif //#if WITH_ENGINE
+	}
+
+	/**
+	 * Return the node type name of the class, we use this when reporting error
+	 */
+	virtual FString GetTypeName() const override
+	{
+		const FString TypeName = TEXT("SkeletonNode");
+		return TypeName;
+	}
+
+	/** Get the class this node want to create */
+	UFUNCTION(BlueprintCallable, Category = "Interchange | Node | Skeleton")
+	virtual class UClass* GetAssetClass() const override
+	{
+		ensure(bIsNodeClassInitialized);
+#if WITH_ENGINE
+		return AssetClass.Get() != nullptr ? AssetClass.Get() : USkeleton::StaticClass();
+#else
+		return nullptr;
+#endif
+	}
+
+	virtual FGuid GetHash() const override
+	{
+		return Attributes->GetStorageHash();
+	}
+
+public:
+	/** Get skeletalMesh asset Uid used to create the data in the post pipeline step.*/
+	UFUNCTION(BlueprintCallable, Category = "Interchange | Node | Skeleton")
+	bool GetCustomSkeletalMeshUid(FString& AttributeValue) const
+	{
+		IMPLEMENT_NODE_ATTRIBUTE_GETTER(SkeletalMeshUid, FString);
+	}
+	/** Set skeletalMesh asset Uid used to create the data in the post pipeline step.*/
+	UFUNCTION(BlueprintCallable, Category = "Interchange | Node | Skeleton")
+	bool SetCustomSkeletalMeshUid(const FString& AttributeValue)
+	{
+		IMPLEMENT_NODE_ATTRIBUTE_SETTER_NODELEGATE(SkeletalMeshUid, FString)
+	}
+
+private:
+
+	void FillAssetClassFromAttribute()
+	{
+#if WITH_ENGINE
+		FString OperationName = GetTypeName() + TEXT(".GetAssetClassName");
+		FString ClassName;
+		InterchangePrivateNodeBase::GetCustomAttribute<FString>(*Attributes, ClassNameAttributeKey, OperationName, ClassName);
+		if (ClassName.Equals(UPhysicsAsset::StaticClass()->GetName()))
+		{
+			AssetClass = UPhysicsAsset::StaticClass();
+			bIsNodeClassInitialized = true;
+		}
+#endif
+	}
+
+	bool SetNodeClassFromClassAttribute()
+	{
+		if (!bIsNodeClassInitialized)
+		{
+			FillAssetClassFromAttribute();
+		}
+		return bIsNodeClassInitialized;
+	}
+
+	bool IsEditorOnlyDataDefined()
+	{
+#if WITH_EDITORONLY_DATA
+		return true;
+#else
+		return false;
+#endif
+	}
+
+	const UE::Interchange::FAttributeKey ClassNameAttributeKey = UE::Interchange::FBaseNodeStaticData::ClassTypeAttributeKey();
+	//SkeletalMesh asset Uid use to create the data in the post pipeline step
+	const UE::Interchange::FAttributeKey Macro_CustomSkeletalMeshUidKey = UE::Interchange::FAttributeKey(TEXT("SkeletalMeshUid"));
+	
+protected:
+#if WITH_ENGINE
+	TSubclassOf<UPhysicsAsset> AssetClass = nullptr;
+#endif
+	bool bIsNodeClassInitialized = false;
+};
diff --git a/Engine/Plugins/Runtime/Interchange/Source/Nodes/Public/InterchangeSkeletalMeshFactoryNode.h b/Engine/Plugins/Runtime/Interchange/Source/Nodes/Public/InterchangeSkeletalMeshFactoryNode.h
--- a/Engine/Plugins/Runtime/Interchange/Source/Nodes/Public/InterchangeSkeletalMeshFactoryNode.h
+++ b/Engine/Plugins/Runtime/Interchange/Source/Nodes/Public/InterchangeSkeletalMeshFactoryNode.h
@@ -147,7 +147,58 @@ public:
 	{
 		return LodDependencies.RemoveName(LodDataUniqueId);
 	}
+	
+	/** Return false if the Attribute was not set previously.*/
+	UFUNCTION(BlueprintCallable, Category = "Interchange | Node | SkeletalMesh")
+	bool GetCustomSkeletonSoftObjectPath(FSoftObjectPath& AttributeValue) const
+	{
+		IMPLEMENT_NODE_ATTRIBUTE_GETTER(SkeletonSoftObjectPath, FSoftObjectPath)
+	}
+
+	UFUNCTION(BlueprintCallable, Category = "Interchange | Node | SkeletalMesh")
+	bool SetCustomSkeletonSoftObjectPath(const FSoftObjectPath& AttributeValue)
+	{
+		IMPLEMENT_NODE_ATTRIBUTE_SETTER_NODELEGATE(SkeletonSoftObjectPath, FSoftObjectPath)
+	}
+
+	/** Return false if the Attribute was not set previously.*/
+	UFUNCTION(BlueprintCallable, Category = "Interchange | Node | SkeletalMesh")
+	bool GetCustomImportMorphTarget(bool& AttributeValue) const
+	{
+		IMPLEMENT_NODE_ATTRIBUTE_GETTER(ImportMorphTarget, bool)
+	}
+
+	UFUNCTION(BlueprintCallable, Category = "Interchange | Node | SkeletalMesh")
+	bool SetCustomImportMorphTarget(const bool& AttributeValue)
+	{
+		IMPLEMENT_NODE_ATTRIBUTE_SETTER_NODELEGATE(ImportMorphTarget, bool)
+	}
 
+	/** Return false if the Attribute was not set previously.*/
+	UFUNCTION(BlueprintCallable, Category = "Interchange | Node | SkeletalMesh")
+	bool GetCustomCreatePhysicsAsset(bool& AttributeValue) const
+	{
+		IMPLEMENT_NODE_ATTRIBUTE_GETTER(CreatePhysicsAsset, bool)
+	}
+
+	UFUNCTION(BlueprintCallable, Category = "Interchange | Node | SkeletalMesh")
+	bool SetCustomCreatePhysicsAsset(const bool& AttributeValue)
+	{
+		IMPLEMENT_NODE_ATTRIBUTE_SETTER_NODELEGATE(CreatePhysicsAsset, bool)
+	}
+
+	/** Return false if the Attribute was not set previously.*/
+	UFUNCTION(BlueprintCallable, Category = "Interchange | Node | SkeletalMesh")
+	bool GetCustomPhysicAssetSoftObjectPath(FSoftObjectPath& AttributeValue) const
+	{
+		IMPLEMENT_NODE_ATTRIBUTE_GETTER(PhysicAssetSoftObjectPath, FSoftObjectPath)
+	}
+
+	UFUNCTION(BlueprintCallable, Category = "Interchange | Node | SkeletalMesh")
+	bool SetCustomPhysicAssetSoftObjectPath(const FSoftObjectPath& AttributeValue)
+	{
+		IMPLEMENT_NODE_ATTRIBUTE_SETTER_NODELEGATE(PhysicAssetSoftObjectPath, FSoftObjectPath)
+	}
 private:
 
 	void FillAssetClassFromAttribute()
@@ -183,6 +234,10 @@ private:
 	}
 
 	const UE::Interchange::FAttributeKey ClassNameAttributeKey = UE::Interchange::FBaseNodeStaticData::ClassTypeAttributeKey();
+	const UE::Interchange::FAttributeKey Macro_CustomImportMorphTargetKey = UE::Interchange::FAttributeKey(TEXT("ImportMorphTarget"));
+	const UE::Interchange::FAttributeKey Macro_CustomSkeletonSoftObjectPathKey = UE::Interchange::FAttributeKey(TEXT("SkeletonSoftObjectPath"));
+	const UE::Interchange::FAttributeKey Macro_CustomCreatePhysicsAssetKey = UE::Interchange::FAttributeKey(TEXT("CreatePhysicsAsset"));
+	const UE::Interchange::FAttributeKey Macro_CustomPhysicAssetSoftObjectPathKey = UE::Interchange::FAttributeKey(TEXT("PhysicAssetSoftObjectPath"));
 
 	UE::Interchange::FNameAttributeArrayHelper LodDependencies;
 protected:
diff --git a/Engine/Plugins/Runtime/Interchange/Source/Pipelines/InterchangePipelines.Build.cs b/Engine/Plugins/Runtime/Interchange/Source/Pipelines/InterchangePipelines.Build.cs
--- a/Engine/Plugins/Runtime/Interchange/Source/Pipelines/InterchangePipelines.Build.cs
+++ b/Engine/Plugins/Runtime/Interchange/Source/Pipelines/InterchangePipelines.Build.cs
@@ -17,6 +17,16 @@ namespace UnrealBuildTool.Rules
 					"InterchangeNodes"
 				}
 			);
+
+			if (Target.Type == TargetType.Editor)
+			{
+				PrivateDependencyModuleNames.AddRange(
+					new string[]
+					{
+						"PhysicsUtilities",
+					}
+				);
+			}
 		}
 	}
 }
diff --git a/Engine/Plugins/Runtime/Interchange/Source/Pipelines/Private/InterchangeGenericAssetsPipeline.cpp b/Engine/Plugins/Runtime/Interchange/Source/Pipelines/Private/InterchangeGenericAssetsPipeline.cpp
--- a/Engine/Plugins/Runtime/Interchange/Source/Pipelines/Private/InterchangeGenericAssetsPipeline.cpp
+++ b/Engine/Plugins/Runtime/Interchange/Source/Pipelines/Private/InterchangeGenericAssetsPipeline.cpp
@@ -86,11 +86,15 @@ namespace UE::Interchange::Private
 	}
 }
 
+void UInterchangeGenericAssetsPipeline::PreDialogCleanup(const FName PipelineStackName)
+{
+	//We always clean the pipeline skeleton when showing the dialog
+	Skeleton = nullptr;
+	SaveSettings(PipelineStackName);
+}
 
 bool UInterchangeGenericAssetsPipeline::ExecutePreImportPipeline(UInterchangeBaseNodeContainer* InBaseNodeContainer, const TArray<UInterchangeSourceData*>& InSourceDatas)
 {
-	// @TODO: break up this method into multiple smaller ones, before it becomes more of a monster!
-
 	if (!InBaseNodeContainer)
 	{
 		UE_LOG(LogInterchangePipeline, Warning, TEXT("UInterchangeGenericAssetsPipeline: Cannot execute pre-import pipeline because InBaseNodeContrainer is null"));
@@ -105,12 +109,9 @@ bool UInterchangeGenericAssetsPipeline::ExecutePreImportPipeline(UInterchangeBas
 	}
 
 	PipelineMeshesUtilities = UInterchangePipelineMeshesUtilities::CreateInterchangePipelineMeshesUtilities(BaseNodeContainer);
-	TArray<UInterchangeMeshNode*> SkinnedMeshNodes;
-	TArray<UInterchangeMeshNode*> StaticMeshNodes;
-	TMap<FString, TArray<FString>> SkeletalMeshFactoryDependencyOrderPerSkeletonRootNodeUid;
 	
 	//Find all translated node we need for this pipeline
-	BaseNodeContainer->IterateNodes([this, &SkinnedMeshNodes, &StaticMeshNodes](const FString& NodeUid, UInterchangeBaseNode* Node)
+	BaseNodeContainer->IterateNodes([this](const FString& NodeUid, UInterchangeBaseNode* Node)
 	{
 		switch(Node->GetnodeContainerType())
 		{
@@ -124,27 +125,6 @@ bool UInterchangeGenericAssetsPipeline::ExecutePreImportPipeline(UInterchangeBas
 				{
 					MaterialNodes.Add(MaterialNode);
 				}
-				else if (UInterchangeMeshNode* MeshNode = Cast<UInterchangeMeshNode>(Node))
-				{
-
-					MeshNodes.Add(MeshNode);
-					if (MeshNode->IsSkinnedMesh())
-					{
-						SkinnedMeshNodes.Add(MeshNode);
-					}
-					else
-					{
-						StaticMeshNodes.Add(MeshNode);
-					}
-				}
-			}
-			break;
-			case EInterchangeNodeContainerType::NodeContainerType_TranslatedScene:
-			{
-				if (UInterchangeSceneNode* SceneNode = Cast<UInterchangeSceneNode>(Node))
-				{
-					SceneNodes.Add(SceneNode);
-				}
 			}
 			break;
 		}
@@ -168,169 +148,10 @@ bool UInterchangeGenericAssetsPipeline::ExecutePreImportPipeline(UInterchangeBas
 		}
 	}
 
+	//Create skeletalmesh factory nodes
+	ExecutePreImportPipelineSkeletalMesh();
 
-	if (bImportSkeletalMeshes && SkinnedMeshNodes.Num() > 0)
-	{
-		auto SetSkeletalMeshDependencies = [&SkeletalMeshFactoryDependencyOrderPerSkeletonRootNodeUid](const FString& JointNodeUid, UInterchangeSkeletalMeshFactoryNode* SkeletalMeshFactoryNode)
-		{
-			TArray<FString>& SkeletalMeshFactoryDependencyOrder = SkeletalMeshFactoryDependencyOrderPerSkeletonRootNodeUid.FindOrAdd(JointNodeUid);
-			//Updating the skeleton is not multi thread safe, so we add dependency between skeletalmesh altering the same skeleton
-			//TODO make the skeletalMesh ReferenceSkeleton thread safe to allow multiple parallel skeletalmesh factory on the same skeleton asset.
-			int32 DependencyIndex = SkeletalMeshFactoryDependencyOrder.AddUnique(SkeletalMeshFactoryNode->GetUniqueID());
-			if (DependencyIndex > 0)
-			{
-				const FString SkeletalMeshFactoryNodeDependencyUid = SkeletalMeshFactoryDependencyOrder[DependencyIndex - 1];
-				SkeletalMeshFactoryNode->AddFactoryDependencyUid(SkeletalMeshFactoryNodeDependencyUid);
-			}
-		};
-
-		if (bCombineSkeletalMeshes)
-		{
-			//////////////////////////////////////////////////////////////////////////
-			//Combined everything we can
-			TMap<FString, TArray<FString>> MeshUidsPerSkeletonRootUid;
-			auto CreatePerSkeletonRootUidCombinedSkinnedMesh = [this, &MeshUidsPerSkeletonRootUid, &SetSkeletalMeshDependencies](const bool bUseInstanceMesh)
-			{
-				bool bFoundInstances = false;
-				for (const TPair<FString, TArray<FString>>& SkeletonRootUidAndMeshUids : MeshUidsPerSkeletonRootUid)
-				{
-					const FString& SkeletonRootUid = SkeletonRootUidAndMeshUids.Key;
-					//Every iteration is a skeletalmesh asset that combine all MeshInstances sharing the same skeleton root node
-					UInterchangeSkeletonFactoryNode* SkeletonFactoryNode = CreateSkeletonFactoryNode(SkeletonRootUid);
-					//The MeshUids can represent a SceneNode pointing on a MeshNode or directly a MeshNode;
-					TMap<int32, TArray<FString>> MeshUidsPerLodIndex;
-					const TArray<FString>& MeshUids = SkeletonRootUidAndMeshUids.Value;
-					for (const FString& MeshUid : MeshUids)
-					{
-						if (bUseInstanceMesh)
-						{
-							const FInterchangeMeshInstance& MeshInstance = PipelineMeshesUtilities->GetMeshInstanceByUid(MeshUid);
-							for (const TPair<int32, FInterchangeLodSceneNodeContainer>& LodIndexAndSceneNodeContainer : MeshInstance.SceneNodePerLodIndex)
-							{
-								const int32 LodIndex = LodIndexAndSceneNodeContainer.Key;
-								const FInterchangeLodSceneNodeContainer& SceneNodeContainer = LodIndexAndSceneNodeContainer.Value;
-								TArray<FString>& TranslatedNodes = MeshUidsPerLodIndex.FindOrAdd(LodIndex);
-								for (const UInterchangeSceneNode* SceneNode : SceneNodeContainer.SceneNodes)
-								{
-									TranslatedNodes.Add(SceneNode->GetUniqueID());
-								}
-							}
-						}
-						else
-						{
-							//MeshGeometry cannot have Lod since LODs are define in the scene node
-							const FInterchangeMeshGeometry& MeshGeometry = PipelineMeshesUtilities->GetMeshGeometryByUid(MeshUid);
-							const int32 LodIndex = 0;
-							TArray<FString>& TranslatedNodes = MeshUidsPerLodIndex.FindOrAdd(LodIndex);
-							TranslatedNodes.Add(MeshGeometry.MeshUid);
-						}
-					}
-
-					if (MeshUidsPerLodIndex.Num() > 0)
-					{
-						UInterchangeSkeletalMeshFactoryNode* SkeletalMeshFactoryNode = CreateSkeletalMeshFactoryNode(SkeletonRootUid, MeshUidsPerLodIndex);
-						SetSkeletalMeshDependencies(SkeletonRootUid, SkeletalMeshFactoryNode);
-						SkeletonFactoryNodes.Add(SkeletonFactoryNode);
-						SkeletalMeshFactoryNodes.Add(SkeletalMeshFactoryNode);
-						bFoundInstances = true;
-					}
-				}
-				return bFoundInstances;
-			};
-
-			bool bFoundMeshes = false;
-			if (bBakeMeshes)
-			{
-				PipelineMeshesUtilities->GetCombinedSkinnedMeshInstances(MeshUidsPerSkeletonRootUid);
-				const bool bUseMeshInstance = true;
-				bFoundMeshes = CreatePerSkeletonRootUidCombinedSkinnedMesh(bUseMeshInstance);
-			}
-
-			if (!bFoundMeshes)
-			{
-				MeshUidsPerSkeletonRootUid.Empty();
-				PipelineMeshesUtilities->GetCombinedSkinnedMeshGeometries(MeshUidsPerSkeletonRootUid);
-				const bool bUseMeshInstance = false;
-				CreatePerSkeletonRootUidCombinedSkinnedMesh(bUseMeshInstance);
-			}
-		}
-		else
-		{
-			//////////////////////////////////////////////////////////////////////////
-			//Do not combined meshes
-			TArray<FString> MeshUids;
-			auto CreatePerSkeletonRootUidSkinnedMesh = [this, &MeshUids, &SetSkeletalMeshDependencies](const bool bUseInstanceMesh)
-			{
-				bool bFoundInstances = false;
-				for (const FString& MeshUid : MeshUids)
-				{
-					//Every iteration is a skeletalmesh asset that combine all MeshInstances sharing the same skeleton root node
-					//The MeshUids can represent a SceneNode pointing on a MeshNode or directly a MeshNode;
-					TMap<int32, TArray<FString>> MeshUidsPerLodIndex;
-					FString SkeletonRootUid;
-					if (!(bUseInstanceMesh ? PipelineMeshesUtilities->IsValidMeshInstanceUid(MeshUid) : PipelineMeshesUtilities->IsValidMeshGeometryUid(MeshUid)))
-					{
-						continue;
-					}
-					SkeletonRootUid = (bUseInstanceMesh ? PipelineMeshesUtilities->GetMeshInstanceSkeletonRootUid(MeshUid) : PipelineMeshesUtilities->GetMeshGeometrySkeletonRootUid(MeshUid));
-					if (SkeletonRootUid.IsEmpty())
-					{
-						//Log an error
-						continue;
-					}
-					UInterchangeSkeletonFactoryNode* SkeletonFactoryNode = CreateSkeletonFactoryNode(SkeletonRootUid);
-					if (bUseInstanceMesh)
-					{
-						const FInterchangeMeshInstance& MeshInstance = PipelineMeshesUtilities->GetMeshInstanceByUid(MeshUid);
-						for (const TPair<int32, FInterchangeLodSceneNodeContainer>& LodIndexAndSceneNodeContainer : MeshInstance.SceneNodePerLodIndex)
-						{
-							const int32 LodIndex = LodIndexAndSceneNodeContainer.Key;
-							const FInterchangeLodSceneNodeContainer& SceneNodeContainer = LodIndexAndSceneNodeContainer.Value;
-							TArray<FString>& TranslatedNodes = MeshUidsPerLodIndex.FindOrAdd(LodIndex);
-							for (const UInterchangeSceneNode* SceneNode : SceneNodeContainer.SceneNodes)
-							{
-								TranslatedNodes.Add(SceneNode->GetUniqueID());
-							}
-						}
-					}
-					else
-					{
-						const FInterchangeMeshGeometry& MeshGeometry = PipelineMeshesUtilities->GetMeshGeometryByUid(MeshUid);
-						const int32 LodIndex = 0;
-						TArray<FString>& TranslatedNodes = MeshUidsPerLodIndex.FindOrAdd(LodIndex);
-						TranslatedNodes.Add(MeshGeometry.MeshUid);
-					}
-					if (MeshUidsPerLodIndex.Num() > 0)
-					{
-						UInterchangeSkeletalMeshFactoryNode* SkeletalMeshFactoryNode = CreateSkeletalMeshFactoryNode(SkeletonRootUid, MeshUidsPerLodIndex);
-						SetSkeletalMeshDependencies(SkeletonRootUid, SkeletalMeshFactoryNode);
-						SkeletonFactoryNodes.Add(SkeletonFactoryNode);
-						SkeletalMeshFactoryNodes.Add(SkeletalMeshFactoryNode);
-						bFoundInstances = true;
-					}
-				}
-				return bFoundInstances;
-			};
-			
-			bool bFoundMeshes = false;
-			if (bBakeMeshes)
-			{
-				PipelineMeshesUtilities->GetAllSkinnedMeshInstance(MeshUids);
-				const bool bUseMeshInstance = true;
-				bFoundMeshes = CreatePerSkeletonRootUidSkinnedMesh(bUseMeshInstance);
-			}
-
-			if (!bFoundMeshes)
-			{
-				MeshUids.Empty();
-				PipelineMeshesUtilities->GetAllSkinnedMeshGeometry(MeshUids);
-				const bool bUseMeshInstance = false;
-				CreatePerSkeletonRootUidSkinnedMesh(bUseMeshInstance);
-			}
-		}
-	}
-
-	if (bImportStaticMeshes && StaticMeshNodes.Num() > 0)
+	if (bImportStaticMeshes)
 	{
 		if (bCombineStaticMeshes)
 		{
@@ -456,33 +277,27 @@ bool UInterchangeGenericAssetsPipeline::ExecutePreImportPipeline(UInterchangeBas
 		}
 	}
 
-	const UClass* SkeletalMeshFactoryNodeClass = UInterchangeSkeletalMeshFactoryNode::StaticClass();
-	TArray<FString> SkeletalMeshNodeUids;
-	BaseNodeContainer->GetNodes(SkeletalMeshFactoryNodeClass, SkeletalMeshNodeUids);
-
- 	const UClass* StaticMeshFactoryNodeClass = UInterchangeStaticMeshFactoryNode::StaticClass();
- 	TArray<FString> StaticMeshNodeUids;
- 	BaseNodeContainer->GetNodes(StaticMeshFactoryNodeClass, StaticMeshNodeUids);
+	ImplementUseSourceNameForAssetOption();
 
-	//TODO count also the imported animations
+	return true;
+}
 
-	//If we import only one asset, and bUseSourceNameForAsset is true, we want to rename the asset using the file name.
-	const int32 MeshesAndAnimsImportedNodeCount = SkeletalMeshNodeUids.Num() + StaticMeshNodeUids.Num();
-	if (bUseSourceNameForAsset && MeshesAndAnimsImportedNodeCount == 1)
+bool UInterchangeGenericAssetsPipeline::ExecutePostImportPipeline(const UInterchangeBaseNodeContainer* InBaseNodeContainer, const FString& NodeKey, UObject* CreatedAsset)
+{
+	//We do not use the provided base container since ExecutePreImportPipeline cache it
+	//We just make sure the same one is pass in parameter
+	if (!BaseNodeContainer || !ensure(BaseNodeContainer == InBaseNodeContainer) || !CreatedAsset)
 	{
-		if (SkeletalMeshNodeUids.Num() > 0)
-		{
-			UInterchangeSkeletalMeshFactoryNode* SkeletalMeshNode = Cast<UInterchangeSkeletalMeshFactoryNode>(BaseNodeContainer->GetNode(SkeletalMeshNodeUids[0]));
-			const FString DisplayLabelName = FPaths::GetBaseFilename(SourceDatas[0]->GetFilename());
-			SkeletalMeshNode->SetDisplayLabel(DisplayLabelName);
-		}
-		else if (StaticMeshNodeUids.Num() > 0)
-		{
-			UInterchangeStaticMeshFactoryNode* StaticMeshNode = Cast<UInterchangeStaticMeshFactoryNode>(BaseNodeContainer->GetNode(StaticMeshNodeUids[0]));
-			const FString DisplayLabelName = FPaths::GetBaseFilename(SourceDatas[0]->GetFilename());
-			StaticMeshNode->SetDisplayLabel(DisplayLabelName);
-		}
+		return false;
 	}
+	UInterchangeBaseNode* Node = BaseNodeContainer->GetNode(NodeKey);
+	if (!Node)
+	{
+		return false;
+	}
+
+	//Finish the physics asset import, it need the skeletal mesh render data to create the physics collision geometry
+	PostImportPhysicsAssetImport(CreatedAsset, Node);
 
 	return true;
 }
@@ -574,43 +389,37 @@ UInterchangeMaterialFactoryNode* UInterchangeGenericAssetsPipeline::CreateMateri
 	return MaterialFactoryNode;
 }
 
-UInterchangeSkeletonFactoryNode* UInterchangeGenericAssetsPipeline::CreateSkeletonFactoryNode(const FString& RootJointUid)
+
+void UInterchangeGenericAssetsPipeline::ImplementUseSourceNameForAssetOption()
 {
-	const UInterchangeBaseNode* RootJointNode = BaseNodeContainer->GetNode(RootJointUid);
-	if (!RootJointNode)
-	{
-		return nullptr;
-	}
-	FString DisplayLabel = RootJointNode->GetDisplayLabel() + TEXT("_Skeleton");
-	FString SkeletonUid = TEXT("\\Skeleton\\") + RootJointNode->GetUniqueID();
+	const UClass* SkeletalMeshFactoryNodeClass = UInterchangeSkeletalMeshFactoryNode::StaticClass();
+	TArray<FString> SkeletalMeshNodeUids;
+	BaseNodeContainer->GetNodes(SkeletalMeshFactoryNodeClass, SkeletalMeshNodeUids);
 
-	UInterchangeSkeletonFactoryNode* SkeletonFactoryNode = nullptr;
-	if (BaseNodeContainer->IsNodeUidValid(SkeletonUid))
-	{
-		//The node already exist, just return it
-		SkeletonFactoryNode = Cast<UInterchangeSkeletonFactoryNode>(BaseNodeContainer->GetNode(SkeletonUid));
-		if (!ensure(SkeletonFactoryNode))
-		{
-			//Log an error
-			return nullptr;
-		}
-	}
-	else
+	const UClass* StaticMeshFactoryNodeClass = UInterchangeStaticMeshFactoryNode::StaticClass();
+	TArray<FString> StaticMeshNodeUids;
+	BaseNodeContainer->GetNodes(StaticMeshFactoryNodeClass, StaticMeshNodeUids);
+
+	//TODO count also the imported animations
+
+	//If we import only one asset, and bUseSourceNameForAsset is true, we want to rename the asset using the file name.
+	const int32 MeshesAndAnimsImportedNodeCount = SkeletalMeshNodeUids.Num() + StaticMeshNodeUids.Num();
+
+	//Skeletalmesh must always rename created skeleton and physics asset to (SKNAME_Skeleton of SKNAME_PhysicsAsset). So the pipeline option
+	// bUseSourceNameForAsset will change or not the SKNAME.
+	ImplementUseSourceNameForAssetOptionSkeletalMesh(MeshesAndAnimsImportedNodeCount);
+
+	if (bUseSourceNameForAsset && MeshesAndAnimsImportedNodeCount == 1)
 	{
-		SkeletonFactoryNode = NewObject<UInterchangeSkeletonFactoryNode>(BaseNodeContainer, NAME_None);
-		if (!ensure(SkeletonFactoryNode))
+		if (StaticMeshNodeUids.Num() > 0)
 		{
-			return nullptr;
+			UInterchangeStaticMeshFactoryNode* StaticMeshNode = Cast<UInterchangeStaticMeshFactoryNode>(BaseNodeContainer->GetNode(StaticMeshNodeUids[0]));
+			const FString DisplayLabelName = FPaths::GetBaseFilename(SourceDatas[0]->GetFilename());
+			StaticMeshNode->SetDisplayLabel(DisplayLabelName);
 		}
-		SkeletonFactoryNode->InitializeSkeletonNode(SkeletonUid, DisplayLabel, USkeleton::StaticClass()->GetName());
-		SkeletonFactoryNode->SetCustomRootJointUid(RootJointNode->GetUniqueID());
-		SkeletonFactoryNode->AddTargetAssetUid(RootJointNode->GetUniqueID());	// todo: add all joints as target assets?
-		BaseNodeContainer->AddNode(SkeletonFactoryNode);
 	}
-	return SkeletonFactoryNode;
-}
-
 
+}
 bool UInterchangeGenericAssetsPipeline::MakeMeshFactoryNodeUidAndDisplayLabel(const TMap<int32, TArray<FString>>& MeshUidsPerLodIndex, int32 LodIndex, FString& NewNodeUid, FString& DisplayLabel)
 {
 	int32 SceneNodeCount = 0;
@@ -669,142 +478,6 @@ bool UInterchangeGenericAssetsPipeline::MakeMeshFactoryNodeUidAndDisplayLabel(co
 	return false;
 }
 
-
-UInterchangeSkeletalMeshFactoryNode* UInterchangeGenericAssetsPipeline::CreateSkeletalMeshFactoryNode(const FString& RootJointUid, const TMap<int32, TArray<FString>>& MeshUidsPerLodIndex)
-{
-	//Get the skeleton factory node
-	const UInterchangeBaseNode* RootJointNode = BaseNodeContainer->GetNode(RootJointUid);
-	if (!RootJointNode)
-	{
-		return nullptr;
-	}
-	const FString SkeletonUid = TEXT("\\Skeleton\\") + RootJointNode->GetUniqueID();
-	const UInterchangeSkeletonFactoryNode* SkeletonFactoryNode = Cast<UInterchangeSkeletonFactoryNode>(BaseNodeContainer->GetNode(SkeletonUid));
-	if (!ensure(SkeletonFactoryNode))
-	{
-		//Log an error
-		return nullptr;
-	}
-	
-	if (MeshUidsPerLodIndex.Num() == 0)
-	{
-		return nullptr;
-	}
-	
-	// Create the skeletal mesh factory node, name it according to the first mesh node compositing the meshes
-	FString SkeletalMeshUid_MeshNamePart;
-	FString DisplayLabel;
-	const int32 BaseLodIndex = 0;
-	if (!MakeMeshFactoryNodeUidAndDisplayLabel(MeshUidsPerLodIndex, BaseLodIndex, SkeletalMeshUid_MeshNamePart, DisplayLabel))
-	{
-		// Log an error
-		return nullptr;
-	}
-
-	const FString SkeletalMeshUid = TEXT("\\SkeletalMesh") + SkeletalMeshUid_MeshNamePart + SkeletonUid;
-	UInterchangeSkeletalMeshFactoryNode* SkeletalMeshFactoryNode = NewObject<UInterchangeSkeletalMeshFactoryNode>(BaseNodeContainer, NAME_None);
-	if (!ensure(SkeletalMeshFactoryNode))
-	{
-		return nullptr;
-	}
-
-	SkeletalMeshFactoryNode->InitializeSkeletalMeshNode(SkeletalMeshUid, DisplayLabel, USkeletalMesh::StaticClass()->GetName());
-	SkeletalMeshFactoryNode->AddFactoryDependencyUid(SkeletonUid);
-	BaseNodeContainer->AddNode(SkeletalMeshFactoryNode);
-
-	AddLodDataToSkeletalMesh(SkeletonFactoryNode, SkeletalMeshFactoryNode, MeshUidsPerLodIndex);
-
-	return SkeletalMeshFactoryNode;
-}
-
-
-UInterchangeSkeletalMeshLodDataNode* UInterchangeGenericAssetsPipeline::CreateSkeletalMeshLodDataNode(const FString& NodeName, const FString& NodeUniqueID)
-{
-	FString DisplayLabel(NodeName);
-	FString NodeUID(NodeUniqueID);
-	UInterchangeSkeletalMeshLodDataNode* SkeletalMeshLodDataNode = NewObject<UInterchangeSkeletalMeshLodDataNode>(BaseNodeContainer, NAME_None);
-	if (!ensure(SkeletalMeshLodDataNode))
-	{
-		//TODO log error
-		return nullptr;
-	}
-
-	SkeletalMeshLodDataNode->InitializeNode(NodeUID, DisplayLabel, EInterchangeNodeContainerType::NodeContainerType_FactoryData);
-	BaseNodeContainer->AddNode(SkeletalMeshLodDataNode);
-	return SkeletalMeshLodDataNode;
-}
-
-
-void UInterchangeGenericAssetsPipeline::AddLodDataToSkeletalMesh(const UInterchangeSkeletonFactoryNode* SkeletonFactoryNode, UInterchangeSkeletalMeshFactoryNode* SkeletalMeshFactoryNode, const TMap<int32, TArray<FString>>& NodeUidsPerLodIndex)
-{
-	const FString SkeletalMeshUid = SkeletalMeshFactoryNode->GetUniqueID();
-	const FString SkeletonUid = SkeletonFactoryNode->GetUniqueID();
-	FString RootjointNodeUid;
-	SkeletonFactoryNode->GetCustomRootJointUid(RootjointNodeUid);
-	for (const TPair<int32, TArray<FString>>& LodIndexAndNodeUids : NodeUidsPerLodIndex)
-	{
-		const int32 LodIndex = LodIndexAndNodeUids.Key;
-		if (!bImportLods && LodIndex > 0)
-		{
-			//If the pipeline should not import lods, skip any lod over base lod
-			continue;
-		}
-		const TArray<FString>& NodeUids = LodIndexAndNodeUids.Value;
-
-		//Create a lod data node with all the meshes for this LOD
-		const FString SkeletalMeshLodDataName = TEXT("LodData") + FString::FromInt(LodIndex);
-		const FString LODDataPrefix = TEXT("\\LodData") + (LodIndex > 0 ? FString::FromInt(LodIndex) : TEXT(""));
-		const FString SkeletalMeshLodDataUniqueID = LODDataPrefix + SkeletalMeshUid + SkeletonUid;
-		//The LodData already exist
-		UInterchangeSkeletalMeshLodDataNode* LodDataNode = Cast<UInterchangeSkeletalMeshLodDataNode>(BaseNodeContainer->GetNode(SkeletalMeshLodDataUniqueID));
-		if (!LodDataNode)
-		{
-			//Add the data for the LOD (skeleton Unique ID and all the mesh node fbx path, so we can find them when we will create the payload data)
-			LodDataNode = CreateSkeletalMeshLodDataNode(SkeletalMeshLodDataName, SkeletalMeshLodDataUniqueID);
-			LodDataNode->SetParentUid(SkeletalMeshUid);
-			LodDataNode->SetCustomSkeletonUid(SkeletonUid);
-			SkeletalMeshFactoryNode->AddLodDataUniqueId(SkeletalMeshLodDataUniqueID);
-		}
-		for (const FString& NodeUid : NodeUids)
-		{
-			TArray<FString> MaterialDependencies;
-			if (const UInterchangeSceneNode* SceneNode = Cast<UInterchangeSceneNode>(BaseNodeContainer->GetNode(NodeUid)))
-			{
-				FString MeshDependency;
-				SceneNode->GetCustomMeshDependencyUid(MeshDependency);
-				if (BaseNodeContainer->IsNodeUidValid(MeshDependency))
-				{
-					SkeletalMeshFactoryNode->AddTargetAssetUid(MeshDependency);
-				}
-
-				SceneNode->GetMaterialDependencyUids(MaterialDependencies);
-			}
-			else if (const UInterchangeMeshNode* MeshNode = Cast<UInterchangeMeshNode>(BaseNodeContainer->GetNode(NodeUid)))
-			{
-				SkeletalMeshFactoryNode->AddTargetAssetUid(NodeUid);
-				MeshNode->GetMaterialDependencies(MaterialDependencies);
-			}
-			const int32 MaterialCount = MaterialDependencies.Num();
-			for (int32 MaterialIndex = 0; MaterialIndex < MaterialCount; ++MaterialIndex)
-			{
-				const FString MaterialFactoryNodeUid = UInterchangeMaterialFactoryNode::GetMaterialFactoryNodeUidFromMaterialNodeUid(MaterialDependencies[MaterialIndex]);
-				if (BaseNodeContainer->IsNodeUidValid(MaterialFactoryNodeUid))
-				{
-					//Create a factory dependency so Material asset are import before the skeletal mesh asset
-					TArray<FString> FactoryDependencies;
-					SkeletalMeshFactoryNode->GetFactoryDependencies(FactoryDependencies);
-					if (!FactoryDependencies.Contains(MaterialFactoryNodeUid))
-					{
-						SkeletalMeshFactoryNode->AddFactoryDependencyUid(MaterialFactoryNodeUid);
-					}
-				}
-			}
-			LodDataNode->AddMeshUid(NodeUid);
-		}
-	}
-}
-
-
 UInterchangeStaticMeshFactoryNode* UInterchangeGenericAssetsPipeline::CreateStaticMeshFactoryNode(const TMap<int32, TArray<FString>>& MeshUidsPerLodIndex)
 {
 	if (MeshUidsPerLodIndex.Num() == 0)
diff --git a/Engine/Plugins/Runtime/Interchange/Source/Pipelines/Private/InterchangeGenericSkeletalMeshPipeline.cpp b/Engine/Plugins/Runtime/Interchange/Source/Pipelines/Private/InterchangeGenericSkeletalMeshPipeline.cpp
new file mode 100644
--- /dev/null
+++ b/Engine/Plugins/Runtime/Interchange/Source/Pipelines/Private/InterchangeGenericSkeletalMeshPipeline.cpp
@@ -0,0 +1,807 @@
+// Copyright Epic Games, Inc. All Rights Reserved. 
+#include "InterchangeGenericAssetsPipeline.h"
+
+#include "Animation/Skeleton.h"
+#include "AssetRegistry/AssetRegistryModule.h"
+#include "CoreMinimal.h"
+#include "InterchangeMaterialFactoryNode.h"
+#include "InterchangeMeshNode.h"
+#include "InterchangePhysicsAssetFactoryNode.h"
+#include "InterchangePipelineLog.h"
+#include "InterchangePipelineMeshesUtilities.h"
+#include "InterchangeSceneNode.h"
+#include "InterchangeSkeletalMeshFactoryNode.h"
+#include "InterchangeSkeletalMeshLodDataNode.h"
+#include "InterchangeSkeletonFactoryNode.h"
+#include "InterchangeSourceData.h"
+#include "Misc/Paths.h"
+#include "Nodes/InterchangeBaseNode.h"
+#include "Nodes/InterchangeBaseNodeContainer.h"
+#if WITH_EDITOR
+#include "PhysicsAssetUtils.h"
+#include "PhysicsEngine/PhysicsAsset.h"
+#endif
+#include "ReferenceSkeleton.h"
+#include "Tasks/Task.h"
+#include "UObject/Object.h"
+#include "UObject/ObjectMacros.h"
+
+namespace UE::Interchange::SkeletalMeshGenericPipeline
+{
+	FName SkeletalLodGetBoneName(const TArray<FMeshBoneInfo>& SkeletalLodRawInfos, int32 BoneIndex)
+	{
+		if (SkeletalLodRawInfos.IsValidIndex(BoneIndex))
+		{
+			return SkeletalLodRawInfos[BoneIndex].Name;
+		}
+		return NAME_None;
+	}
+
+	int32 SkeletalLodFindBoneIndex(const TArray<FMeshBoneInfo>& SkeletalLodRawInfos, FName BoneName)
+	{
+		const int32 BoneCount = SkeletalLodRawInfos.Num();
+		for (int32 BoneIndex = 0; BoneIndex < BoneCount; ++BoneIndex)
+		{
+			if (SkeletalLodRawInfos[BoneIndex].Name == BoneName)
+			{
+				return BoneIndex;
+			}
+		}
+		return INDEX_NONE;
+	}
+
+	int32 SkeletalLodGetParentIndex(const TArray<FMeshBoneInfo>& SkeletalLodRawInfos, int32 BoneIndex)
+	{
+		if (SkeletalLodRawInfos.IsValidIndex(BoneIndex))
+		{
+			return SkeletalLodRawInfos[BoneIndex].ParentIndex;
+		}
+		return INDEX_NONE;
+	}
+
+	bool DoesParentChainMatch(int32 StartBoneIndex, const FReferenceSkeleton& SkeletonRef, const TArray<FMeshBoneInfo>& SkeletalLodRawInfos)
+	{
+		// if start is root bone
+		if (StartBoneIndex == 0)
+		{
+			// verify name of root bone matches
+			return (SkeletonRef.GetBoneName(0) == SkeletalLodGetBoneName(SkeletalLodRawInfos, 0));
+		}
+
+		int32 SkeletonBoneIndex = StartBoneIndex;
+		// If skeleton bone is not found in mesh, fail.
+		int32 MeshBoneIndex = SkeletalLodFindBoneIndex(SkeletalLodRawInfos, SkeletonRef.GetBoneName(SkeletonBoneIndex));
+		if (MeshBoneIndex == INDEX_NONE)
+		{
+			return false;
+		}
+		do
+		{
+			// verify if parent name matches
+			int32 ParentSkeletonBoneIndex = SkeletonRef.GetParentIndex(SkeletonBoneIndex);
+			int32 ParentMeshBoneIndex = SkeletalLodGetParentIndex(SkeletalLodRawInfos, MeshBoneIndex);
+
+			// if one of the parents doesn't exist, make sure both end. Otherwise fail.
+			if ((ParentSkeletonBoneIndex == INDEX_NONE) || (ParentMeshBoneIndex == INDEX_NONE))
+			{
+				return (ParentSkeletonBoneIndex == ParentMeshBoneIndex);
+			}
+
+			// If parents are not named the same, fail.
+			if (SkeletonRef.GetBoneName(ParentSkeletonBoneIndex) != SkeletalLodGetBoneName(SkeletalLodRawInfos, ParentMeshBoneIndex))
+			{
+				return false;
+			}
+
+			// move up
+			SkeletonBoneIndex = ParentSkeletonBoneIndex;
+			MeshBoneIndex = ParentMeshBoneIndex;
+		} while (true);
+
+		return true;
+	}
+
+	void RecursiveBuildSkeletalSkeleton(const FString JoinToAddUid, const int32 ParentIndex, const UInterchangeBaseNodeContainer* BaseNodeContainer, TArray<FMeshBoneInfo>& SkeletalLodRawInfos, TArray<FTransform>& SkeletalLodRawTransforms)
+	{
+		const UInterchangeSceneNode* SceneNode = Cast<UInterchangeSceneNode>(BaseNodeContainer->GetNode(JoinToAddUid));
+		if (!SceneNode || !SceneNode->IsSpecializedTypeContains(FSceneNodeStaticData::GetJointSpecializeTypeString()))
+		{
+			return;
+		}
+
+		int32 JoinIndex = SkeletalLodRawInfos.Num();
+		FMeshBoneInfo& Info = SkeletalLodRawInfos.AddZeroed_GetRef();
+		Info.Name = *SceneNode->GetDisplayLabel();
+		Info.ParentIndex = ParentIndex;
+#if WITH_EDITORONLY_DATA
+		Info.ExportName = Info.Name.ToString();
+#endif
+		FTransform& JoinTransform = SkeletalLodRawTransforms.AddZeroed_GetRef();
+		SceneNode->GetCustomLocalTransform(JoinTransform);
+
+		//Iterate childrens
+		const TArray<FString> ChildrenIds = BaseNodeContainer->GetNodeChildrenUids(JoinToAddUid);
+		for (int32 ChildIndex = 0; ChildIndex < ChildrenIds.Num(); ++ChildIndex)
+		{
+			RecursiveBuildSkeletalSkeleton(ChildrenIds[ChildIndex], JoinIndex, BaseNodeContainer, SkeletalLodRawInfos, SkeletalLodRawTransforms);
+		}
+	}
+
+	bool IsCompatibleSkeleton(const USkeleton* Skeleton, const FString RootJoinUid, const UInterchangeBaseNodeContainer* BaseNodeContainer)
+	{
+		// at least % of bone should match 
+		int32 NumOfBoneMatches = 0;
+		//Make sure the specified Skeleton fit this skeletal mesh
+		const FReferenceSkeleton& SkeletonRef = Skeleton->GetReferenceSkeleton();
+		const int32 SkeletonBoneCount = SkeletonRef.GetRawBoneNum();
+
+		TArray<FMeshBoneInfo> SkeletalLodRawInfos;
+		SkeletalLodRawInfos.Reserve(SkeletonBoneCount);
+		TArray<FTransform> SkeletalLodRawTransforms;
+		SkeletalLodRawTransforms.Reserve(SkeletonBoneCount);
+		RecursiveBuildSkeletalSkeleton(RootJoinUid, INDEX_NONE, BaseNodeContainer, SkeletalLodRawInfos, SkeletalLodRawTransforms);
+		const int32 SkeletalLodBoneCount = SkeletalLodRawInfos.Num();
+
+		// first ensure the parent exists for each bone
+		for (int32 MeshBoneIndex = 0; MeshBoneIndex < SkeletalLodBoneCount; MeshBoneIndex++)
+		{
+			FName MeshBoneName = SkeletalLodRawInfos[MeshBoneIndex].Name;
+			// See if Mesh bone exists in Skeleton.
+			int32 SkeletonBoneIndex = SkeletonRef.FindBoneIndex(MeshBoneName);
+
+			// if found, increase num of bone matches count
+			if (SkeletonBoneIndex != INDEX_NONE)
+			{
+				++NumOfBoneMatches;
+
+				// follow the parent chain to verify the chain is same
+				if (!DoesParentChainMatch(SkeletonBoneIndex, SkeletonRef, SkeletalLodRawInfos))
+				{
+					//Not compatible
+					return false;
+				}
+			}
+			else
+			{
+				int32 CurrentBoneId = MeshBoneIndex;
+				// if not look for parents that matches
+				while (SkeletonBoneIndex == INDEX_NONE && CurrentBoneId != INDEX_NONE)
+				{
+					// find Parent one see exists
+					const int32 ParentMeshBoneIndex = SkeletalLodGetParentIndex(SkeletalLodRawInfos, CurrentBoneId);
+					if (ParentMeshBoneIndex != INDEX_NONE)
+					{
+						// @TODO: make sure RefSkeleton's root ParentIndex < 0 if not, I'll need to fix this by checking TreeBoneIdx
+						FName ParentBoneName = SkeletalLodGetBoneName(SkeletalLodRawInfos, ParentMeshBoneIndex);
+						SkeletonBoneIndex = SkeletonRef.FindBoneIndex(ParentBoneName);
+					}
+
+					// root is reached
+					if (ParentMeshBoneIndex == 0)
+					{
+						break;
+					}
+					else
+					{
+						CurrentBoneId = ParentMeshBoneIndex;
+					}
+				}
+
+				// still no match, return false, no parent to look for
+				if (SkeletonBoneIndex == INDEX_NONE)
+				{
+					return false;
+				}
+
+				// second follow the parent chain to verify the chain is same
+				if (!DoesParentChainMatch(SkeletonBoneIndex, SkeletonRef, SkeletalLodRawInfos))
+				{
+					return false;
+				}
+			}
+		}
+
+		// originally we made sure at least matches more than 50% 
+		// but then slave components can't play since they're only partial
+		// if the hierarchy matches, and if it's more then 1 bone, we allow
+		return (NumOfBoneMatches > 0);
+	}
+}
+
+bool UInterchangeGenericAssetsPipeline::ExecutePreImportPipelineSkeletalMesh()
+{
+	
+	//PipelineMeshesUtilities;
+	TArray<UInterchangeMeshNode*> SkinnedMeshNodes;
+	TArray<UInterchangeMeshNode*> StaticMeshNodes;
+	TMap<FString, TArray<FString>> SkeletalMeshFactoryDependencyOrderPerSkeletonRootNodeUid;
+	
+	//Find all translated node we need for this pipeline
+	BaseNodeContainer->IterateNodes([this, &SkinnedMeshNodes, &StaticMeshNodes](const FString& NodeUid, UInterchangeBaseNode* Node)
+	{
+		switch(Node->GetnodeContainerType())
+		{
+			case EInterchangeNodeContainerType::NodeContainerType_TranslatedAsset:
+			{
+				if (UInterchangeMeshNode* MeshNode = Cast<UInterchangeMeshNode>(Node))
+				{
+					if (MeshNode->IsSkinnedMesh())
+					{
+						SkinnedMeshNodes.Add(MeshNode);
+					}
+					else
+					{
+						StaticMeshNodes.Add(MeshNode);
+					}
+
+				}
+			}
+			break;
+		}
+	});
+
+	if (bImportSkeletalMeshes && SkinnedMeshNodes.Num() > 0)
+	{
+		auto SetSkeletalMeshDependencies = [&SkeletalMeshFactoryDependencyOrderPerSkeletonRootNodeUid](const FString& JointNodeUid, UInterchangeSkeletalMeshFactoryNode* SkeletalMeshFactoryNode)
+		{
+			TArray<FString>& SkeletalMeshFactoryDependencyOrder = SkeletalMeshFactoryDependencyOrderPerSkeletonRootNodeUid.FindOrAdd(JointNodeUid);
+			//Updating the skeleton is not multi thread safe, so we add dependency between skeletalmesh altering the same skeleton
+			//TODO make the skeletalMesh ReferenceSkeleton thread safe to allow multiple parallel skeletalmesh factory on the same skeleton asset.
+			int32 DependencyIndex = SkeletalMeshFactoryDependencyOrder.AddUnique(SkeletalMeshFactoryNode->GetUniqueID());
+			if (DependencyIndex > 0)
+			{
+				const FString SkeletalMeshFactoryNodeDependencyUid = SkeletalMeshFactoryDependencyOrder[DependencyIndex - 1];
+				SkeletalMeshFactoryNode->AddFactoryDependencyUid(SkeletalMeshFactoryNodeDependencyUid);
+			}
+		};
+
+		if (bCombineSkeletalMeshes)
+		{
+			//////////////////////////////////////////////////////////////////////////
+			//Combined everything we can
+			TMap<FString, TArray<FString>> MeshUidsPerSkeletonRootUid;
+			auto CreatePerSkeletonRootUidCombinedSkinnedMesh = [this, &MeshUidsPerSkeletonRootUid, &SetSkeletalMeshDependencies](const bool bUseInstanceMesh)
+			{
+				bool bFoundInstances = false;
+				for (const TPair<FString, TArray<FString>>& SkeletonRootUidAndMeshUids : MeshUidsPerSkeletonRootUid)
+				{
+					const FString& SkeletonRootUid = SkeletonRootUidAndMeshUids.Key;
+					//Every iteration is a skeletalmesh asset that combine all MeshInstances sharing the same skeleton root node
+					UInterchangeSkeletonFactoryNode* SkeletonFactoryNode = CreateSkeletonFactoryNode(SkeletonRootUid);
+					//The MeshUids can represent a SceneNode pointing on a MeshNode or directly a MeshNode;
+					TMap<int32, TArray<FString>> MeshUidsPerLodIndex;
+					const TArray<FString>& MeshUids = SkeletonRootUidAndMeshUids.Value;
+					for (const FString& MeshUid : MeshUids)
+					{
+						if (bUseInstanceMesh)
+						{
+							const FInterchangeMeshInstance& MeshInstance = PipelineMeshesUtilities->GetMeshInstanceByUid(MeshUid);
+							for (const TPair<int32, FInterchangeLodSceneNodeContainer>& LodIndexAndSceneNodeContainer : MeshInstance.SceneNodePerLodIndex)
+							{
+								const int32 LodIndex = LodIndexAndSceneNodeContainer.Key;
+								const FInterchangeLodSceneNodeContainer& SceneNodeContainer = LodIndexAndSceneNodeContainer.Value;
+								TArray<FString>& TranslatedNodes = MeshUidsPerLodIndex.FindOrAdd(LodIndex);
+								for (const UInterchangeSceneNode* SceneNode : SceneNodeContainer.SceneNodes)
+								{
+									TranslatedNodes.Add(SceneNode->GetUniqueID());
+								}
+							}
+						}
+						else
+						{
+							//MeshGeometry cannot have Lod since LODs are define in the scene node
+							const FInterchangeMeshGeometry& MeshGeometry = PipelineMeshesUtilities->GetMeshGeometryByUid(MeshUid);
+							const int32 LodIndex = 0;
+							TArray<FString>& TranslatedNodes = MeshUidsPerLodIndex.FindOrAdd(LodIndex);
+							TranslatedNodes.Add(MeshGeometry.MeshUid);
+						}
+					}
+
+					if (MeshUidsPerLodIndex.Num() > 0)
+					{
+						UInterchangeSkeletalMeshFactoryNode* SkeletalMeshFactoryNode = CreateSkeletalMeshFactoryNode(SkeletonRootUid, MeshUidsPerLodIndex);
+						SetSkeletalMeshDependencies(SkeletonRootUid, SkeletalMeshFactoryNode);
+						SkeletonFactoryNodes.Add(SkeletonFactoryNode);
+						SkeletalMeshFactoryNodes.Add(SkeletalMeshFactoryNode);
+						bFoundInstances = true;
+					}
+				}
+				return bFoundInstances;
+			};
+
+			bool bFoundMeshes = false;
+			if (bBakeMeshes)
+			{
+				PipelineMeshesUtilities->GetCombinedSkinnedMeshInstances(MeshUidsPerSkeletonRootUid);
+				const bool bUseMeshInstance = true;
+				bFoundMeshes = CreatePerSkeletonRootUidCombinedSkinnedMesh(bUseMeshInstance);
+			}
+
+			if (!bFoundMeshes)
+			{
+				MeshUidsPerSkeletonRootUid.Empty();
+				PipelineMeshesUtilities->GetCombinedSkinnedMeshGeometries(MeshUidsPerSkeletonRootUid);
+				const bool bUseMeshInstance = false;
+				CreatePerSkeletonRootUidCombinedSkinnedMesh(bUseMeshInstance);
+			}
+		}
+		else
+		{
+			//////////////////////////////////////////////////////////////////////////
+			//Do not combined meshes
+			TArray<FString> MeshUids;
+			auto CreatePerSkeletonRootUidSkinnedMesh = [this, &MeshUids, &SetSkeletalMeshDependencies](const bool bUseInstanceMesh)
+			{
+				bool bFoundInstances = false;
+				for (const FString& MeshUid : MeshUids)
+				{
+					//Every iteration is a skeletalmesh asset that combine all MeshInstances sharing the same skeleton root node
+					//The MeshUids can represent a SceneNode pointing on a MeshNode or directly a MeshNode;
+					TMap<int32, TArray<FString>> MeshUidsPerLodIndex;
+					FString SkeletonRootUid;
+					if (!(bUseInstanceMesh ? PipelineMeshesUtilities->IsValidMeshInstanceUid(MeshUid) : PipelineMeshesUtilities->IsValidMeshGeometryUid(MeshUid)))
+					{
+						continue;
+					}
+					SkeletonRootUid = (bUseInstanceMesh ? PipelineMeshesUtilities->GetMeshInstanceSkeletonRootUid(MeshUid) : PipelineMeshesUtilities->GetMeshGeometrySkeletonRootUid(MeshUid));
+					if (SkeletonRootUid.IsEmpty())
+					{
+						//Log an error
+						continue;
+					}
+					UInterchangeSkeletonFactoryNode* SkeletonFactoryNode = CreateSkeletonFactoryNode(SkeletonRootUid);
+					if (bUseInstanceMesh)
+					{
+						const FInterchangeMeshInstance& MeshInstance = PipelineMeshesUtilities->GetMeshInstanceByUid(MeshUid);
+						for (const TPair<int32, FInterchangeLodSceneNodeContainer>& LodIndexAndSceneNodeContainer : MeshInstance.SceneNodePerLodIndex)
+						{
+							const int32 LodIndex = LodIndexAndSceneNodeContainer.Key;
+							const FInterchangeLodSceneNodeContainer& SceneNodeContainer = LodIndexAndSceneNodeContainer.Value;
+							TArray<FString>& TranslatedNodes = MeshUidsPerLodIndex.FindOrAdd(LodIndex);
+							for (const UInterchangeSceneNode* SceneNode : SceneNodeContainer.SceneNodes)
+							{
+								TranslatedNodes.Add(SceneNode->GetUniqueID());
+							}
+						}
+					}
+					else
+					{
+						const FInterchangeMeshGeometry& MeshGeometry = PipelineMeshesUtilities->GetMeshGeometryByUid(MeshUid);
+						const int32 LodIndex = 0;
+						TArray<FString>& TranslatedNodes = MeshUidsPerLodIndex.FindOrAdd(LodIndex);
+						TranslatedNodes.Add(MeshGeometry.MeshUid);
+					}
+					if (MeshUidsPerLodIndex.Num() > 0)
+					{
+						UInterchangeSkeletalMeshFactoryNode* SkeletalMeshFactoryNode = CreateSkeletalMeshFactoryNode(SkeletonRootUid, MeshUidsPerLodIndex);
+						SetSkeletalMeshDependencies(SkeletonRootUid, SkeletalMeshFactoryNode);
+						SkeletonFactoryNodes.Add(SkeletonFactoryNode);
+						SkeletalMeshFactoryNodes.Add(SkeletalMeshFactoryNode);
+						bFoundInstances = true;
+					}
+				}
+				return bFoundInstances;
+			};
+			
+			bool bFoundMeshes = false;
+			if (bBakeMeshes)
+			{
+				PipelineMeshesUtilities->GetAllSkinnedMeshInstance(MeshUids);
+				const bool bUseMeshInstance = true;
+				bFoundMeshes = CreatePerSkeletonRootUidSkinnedMesh(bUseMeshInstance);
+			}
+
+			if (!bFoundMeshes)
+			{
+				MeshUids.Empty();
+				PipelineMeshesUtilities->GetAllSkinnedMeshGeometry(MeshUids);
+				const bool bUseMeshInstance = false;
+				CreatePerSkeletonRootUidSkinnedMesh(bUseMeshInstance);
+			}
+		}
+	}
+
+
+	const UClass* SkeletalMeshFactoryNodeClass = UInterchangeSkeletalMeshFactoryNode::StaticClass();
+	TArray<FString> SkeletalMeshNodeUids;
+	BaseNodeContainer->GetNodes(SkeletalMeshFactoryNodeClass, SkeletalMeshNodeUids);
+
+	//If we import only one asset, and bUseSourceNameForAsset is true, we want to rename the asset using the file name.
+	const int32 MeshesAndAnimsImportedNodeCount = SkeletalMeshNodeUids.Num(); // + StaticMeshNodeUids.Num();
+	if (bUseSourceNameForAsset && MeshesAndAnimsImportedNodeCount == 1)
+	{
+		if (SkeletalMeshNodeUids.Num() > 0)
+		{
+			UInterchangeSkeletalMeshFactoryNode* SkeletalMeshNode = Cast<UInterchangeSkeletalMeshFactoryNode>(BaseNodeContainer->GetNode(SkeletalMeshNodeUids[0]));
+			const FString DisplayLabelName = FPaths::GetBaseFilename(SourceDatas[0]->GetFilename());
+			SkeletalMeshNode->SetDisplayLabel(DisplayLabelName);
+		}
+	}
+
+	return true;
+}
+
+
+UInterchangeSkeletonFactoryNode* UInterchangeGenericAssetsPipeline::CreateSkeletonFactoryNode(const FString& RootJointUid)
+{
+	const UInterchangeBaseNode* RootJointNode = BaseNodeContainer->GetNode(RootJointUid);
+	if (!RootJointNode)
+	{
+		return nullptr;
+	}
+	FString DisplayLabel = RootJointNode->GetDisplayLabel() + TEXT("_Skeleton");
+	FString SkeletonUid = TEXT("\\Skeleton\\") + RootJointNode->GetUniqueID();
+
+	UInterchangeSkeletonFactoryNode* SkeletonFactoryNode = nullptr;
+	if (BaseNodeContainer->IsNodeUidValid(SkeletonUid))
+	{
+		//The node already exist, just return it
+		SkeletonFactoryNode = Cast<UInterchangeSkeletonFactoryNode>(BaseNodeContainer->GetNode(SkeletonUid));
+		if (!ensure(SkeletonFactoryNode))
+		{
+			//Log an error
+			return nullptr;
+		}
+	}
+	else
+	{
+		SkeletonFactoryNode = NewObject<UInterchangeSkeletonFactoryNode>(BaseNodeContainer, NAME_None);
+		if (!ensure(SkeletonFactoryNode))
+		{
+			return nullptr;
+		}
+		SkeletonFactoryNode->InitializeSkeletonNode(SkeletonUid, DisplayLabel, USkeleton::StaticClass()->GetName());
+		SkeletonFactoryNode->SetCustomRootJointUid(RootJointNode->GetUniqueID());
+		BaseNodeContainer->AddNode(SkeletonFactoryNode);
+	}
+
+	//If we have a specified skeleton
+	if (Skeleton)
+	{
+		SkeletonFactoryNode->SetEnabled(false);
+		SkeletonFactoryNode->ReferenceObject = Skeleton;
+	}
+	return SkeletonFactoryNode;
+}
+
+UInterchangeSkeletalMeshFactoryNode* UInterchangeGenericAssetsPipeline::CreateSkeletalMeshFactoryNode(const FString& RootJointUid, const TMap<int32, TArray<FString>>& MeshUidsPerLodIndex)
+{
+	//Get the skeleton factory node
+	const UInterchangeBaseNode* RootJointNode = BaseNodeContainer->GetNode(RootJointUid);
+	if (!RootJointNode)
+	{
+		return nullptr;
+	}
+	const FString SkeletonUid = TEXT("\\Skeleton\\") + RootJointNode->GetUniqueID();
+	UInterchangeSkeletonFactoryNode* SkeletonFactoryNode = Cast<UInterchangeSkeletonFactoryNode>(BaseNodeContainer->GetNode(SkeletonUid));
+	if (!ensure(SkeletonFactoryNode))
+	{
+		//Log an error
+		return nullptr;
+	}
+	
+	if (MeshUidsPerLodIndex.Num() == 0)
+	{
+		return nullptr;
+	}
+	
+	auto GetFirstNodeInfo = [this, &MeshUidsPerLodIndex](const int32 Index, FString& OutFirstMeshNodeUid, int32& OutSceneNodeCount)->const UInterchangeBaseNode*
+	{
+		OutSceneNodeCount = 0;
+		if (!ensure(Index >= 0 && MeshUidsPerLodIndex.Num() > Index))
+		{
+			//Log an error
+			return nullptr;
+		}
+		for (const TPair<int32, TArray<FString>>& LodIndexAndMeshUids : MeshUidsPerLodIndex)
+		{
+			if (Index == LodIndexAndMeshUids.Key)
+			{
+				const TArray<FString>& MeshUids = LodIndexAndMeshUids.Value;
+				if (MeshUids.Num() > 0)
+				{
+					const FString& MeshUid = MeshUids[0];
+					UInterchangeMeshNode* MeshNode = Cast<UInterchangeMeshNode>(BaseNodeContainer->GetNode(MeshUid));
+					if (MeshNode)
+					{
+						OutFirstMeshNodeUid = MeshUid;
+						return MeshNode;
+					}
+					UInterchangeSceneNode* SceneNode = Cast<UInterchangeSceneNode>(BaseNodeContainer->GetNode(MeshUid));
+					if (SceneNode)
+					{
+						FString MeshNodeUid;
+						if (SceneNode->GetCustomMeshDependencyUid(MeshNodeUid))
+						{
+							OutSceneNodeCount = MeshUids.Num();
+							OutFirstMeshNodeUid = MeshNodeUid;
+							return SceneNode;
+						}
+					}
+				}
+				//We found the lod but there is no valid Mesh node to return the Uid
+				break;
+			}
+		}
+		return nullptr;
+	};
+
+	FString FirstMeshNodeUid;
+	const int32 BaseLodIndex = 0;
+	int32 SceneNodeCount = 0;
+	const UInterchangeBaseNode* InterchangeBaseNode = GetFirstNodeInfo(BaseLodIndex, FirstMeshNodeUid, SceneNodeCount);
+	if (!InterchangeBaseNode)
+	{
+		//Log an error
+		return nullptr;
+	}
+	const UInterchangeSceneNode* FirstSceneNode = Cast<UInterchangeSceneNode>(InterchangeBaseNode);
+	const UInterchangeMeshNode* FirstMeshNode = Cast<UInterchangeMeshNode>(BaseNodeContainer->GetNode(FirstMeshNodeUid));
+
+	//Create the skeletal mesh factory node, name it according to the first mesh node compositing the meshes
+	FString DisplayLabel = FirstMeshNode->GetDisplayLabel();
+	FString SkeletalMeshUid_MeshNamePart = FirstMeshNodeUid;
+	if(FirstSceneNode)
+	{
+		//If we are instancing one scene node, we want to use it to name the mesh
+		if (SceneNodeCount == 1)
+		{
+			DisplayLabel = FirstSceneNode->GetDisplayLabel();
+		}
+		//Use the first scene node uid this skeletalmesh reference, add backslash since this uid is not asset typed (\\Mesh\\) like FirstMeshNodeUid
+		SkeletalMeshUid_MeshNamePart = TEXT("\\") + FirstSceneNode->GetUniqueID();
+	}
+	const FString SkeletalMeshUid = TEXT("\\SkeletalMesh") + SkeletalMeshUid_MeshNamePart + SkeletonUid;
+	UInterchangeSkeletalMeshFactoryNode* SkeletalMeshFactoryNode = NewObject<UInterchangeSkeletalMeshFactoryNode>(BaseNodeContainer, NAME_None);
+	if (!ensure(SkeletalMeshFactoryNode))
+	{
+		return nullptr;
+	}
+	SkeletalMeshFactoryNode->InitializeSkeletalMeshNode(SkeletalMeshUid, DisplayLabel, USkeletalMesh::StaticClass()->GetName());
+	SkeletalMeshFactoryNode->AddFactoryDependencyUid(SkeletonUid);
+	BaseNodeContainer->AddNode(SkeletalMeshFactoryNode);
+
+	AddLodDataToSkeletalMesh(SkeletonFactoryNode, SkeletalMeshFactoryNode, MeshUidsPerLodIndex);
+	SkeletalMeshFactoryNode->SetCustomImportMorphTarget(bImportMorphTargets);
+
+	//If we have a specified skeleton
+	if (Skeleton)
+	{
+		if (UE::Interchange::SkeletalMeshGenericPipeline::IsCompatibleSkeleton(Skeleton, RootJointNode->GetUniqueID(), BaseNodeContainer))
+		{
+			FSoftObjectPath SkeletonSoftObjectPath(Skeleton.Get());
+			SkeletalMeshFactoryNode->SetCustomSkeletonSoftObjectPath(SkeletonSoftObjectPath);
+		}
+		else
+		{
+			//Log an error, saying we will create a skeleton instead of using the specified one
+			//Should we show a dialog so we are sure the user understand he choose the wrong skeleton
+
+			//Make sure we enable the skeleton factory node
+			SkeletonFactoryNode->SetEnabled(true);
+
+		}
+	}
+
+#if WITH_EDITOR
+	//Physic asset dependency, if we must create or use a specialize physic asset let create
+	//a PhysicsAsset factory node, so the asset will exist when we will setup the skeletalmesh
+	if (bCreatePhysicsAsset)
+	{
+		UInterchangePhysicsAssetFactoryNode* PhysicsAssetFactoryNode = NewObject<UInterchangePhysicsAssetFactoryNode>(BaseNodeContainer, NAME_None);
+		if (ensure(SkeletalMeshFactoryNode))
+		{
+			const FString PhysicsAssetUid = TEXT("\\PhysicsAsset") + SkeletalMeshUid_MeshNamePart + SkeletonUid;
+			const FString PhysicsAssetDisplayLabel = DisplayLabel + TEXT("_PhysicsAsset");
+			PhysicsAssetFactoryNode->InitializePhysicsAssetNode(PhysicsAssetUid, PhysicsAssetDisplayLabel, UPhysicsAsset::StaticClass()->GetName());
+			PhysicsAssetFactoryNode->SetCustomSkeletalMeshUid(SkeletalMeshUid);
+			BaseNodeContainer->AddNode(PhysicsAssetFactoryNode);
+		}
+	}
+	SkeletalMeshFactoryNode->SetCustomCreatePhysicsAsset(bCreatePhysicsAsset);
+	if (!bCreatePhysicsAsset && !PhysicsAsset.IsNull())
+	{
+		FSoftObjectPath PhysicSoftObjectPath(PhysicsAsset.Get());
+		SkeletalMeshFactoryNode->SetCustomPhysicAssetSoftObjectPath(PhysicSoftObjectPath);
+	}
+#endif
+
+	return SkeletalMeshFactoryNode;
+}
+
+UInterchangeSkeletalMeshLodDataNode* UInterchangeGenericAssetsPipeline::CreateSkeletalMeshLodDataNode(const FString& NodeName, const FString& NodeUniqueID)
+{
+	FString DisplayLabel(NodeName);
+	FString NodeUID(NodeUniqueID);
+	UInterchangeSkeletalMeshLodDataNode* SkeletalMeshLodDataNode = NewObject<UInterchangeSkeletalMeshLodDataNode>(BaseNodeContainer, NAME_None);
+	if (!ensure(SkeletalMeshLodDataNode))
+	{
+		//TODO log error
+		return nullptr;
+	}
+	// Creating a UMaterialInterface
+	SkeletalMeshLodDataNode->InitializeNode(NodeUID, DisplayLabel, EInterchangeNodeContainerType::NodeContainerType_FactoryData);
+	BaseNodeContainer->AddNode(SkeletalMeshLodDataNode);
+	return SkeletalMeshLodDataNode;
+}
+
+void UInterchangeGenericAssetsPipeline::AddLodDataToSkeletalMesh(const UInterchangeSkeletonFactoryNode* SkeletonFactoryNode, UInterchangeSkeletalMeshFactoryNode* SkeletalMeshFactoryNode, const TMap<int32, TArray<FString>>& NodeUidsPerLodIndex)
+{
+	const FString SkeletalMeshUid = SkeletalMeshFactoryNode->GetUniqueID();
+	const FString SkeletonUid = SkeletonFactoryNode->GetUniqueID();
+	for (const TPair<int32, TArray<FString>>& LodIndexAndNodeUids : NodeUidsPerLodIndex)
+	{
+		const int32 LodIndex = LodIndexAndNodeUids.Key;
+		if (!bImportLods && LodIndex > 0)
+		{
+			//If the pipeline should not import lods, skip any lod over base lod
+			continue;
+		}
+		const TArray<FString>& NodeUids = LodIndexAndNodeUids.Value;
+
+		//Create a lod data node with all the meshes for this LOD
+		const FString SkeletalMeshLodDataName = TEXT("LodData") + FString::FromInt(LodIndex);
+		const FString LODDataPrefix = TEXT("\\LodData") + (LodIndex > 0 ? FString::FromInt(LodIndex) : TEXT(""));
+		const FString SkeletalMeshLodDataUniqueID = LODDataPrefix + SkeletalMeshUid + SkeletonUid;
+		//The LodData already exist
+		UInterchangeSkeletalMeshLodDataNode* LodDataNode = Cast<UInterchangeSkeletalMeshLodDataNode>(BaseNodeContainer->GetNode(SkeletalMeshLodDataUniqueID));
+		if (!LodDataNode)
+		{
+			//Add the data for the LOD (skeleton Unique ID and all the mesh node fbx path, so we can find them when we will create the payload data)
+			LodDataNode = CreateSkeletalMeshLodDataNode(SkeletalMeshLodDataName, SkeletalMeshLodDataUniqueID);
+			LodDataNode->SetParentUid(SkeletalMeshUid);
+			LodDataNode->SetCustomSkeletonUid(SkeletonUid);
+			SkeletalMeshFactoryNode->AddLodDataUniqueId(SkeletalMeshLodDataUniqueID);
+		}
+		for (const FString& NodeUid : NodeUids)
+		{
+			TArray<FString> MaterialDependencies;
+			if (const UInterchangeSceneNode* SceneNode = Cast<UInterchangeSceneNode>(BaseNodeContainer->GetNode(NodeUid)))
+			{
+				SceneNode->GetMaterialDependencyUids(MaterialDependencies);
+			}
+			else if (const UInterchangeMeshNode* MeshNode = Cast<UInterchangeMeshNode>(BaseNodeContainer->GetNode(NodeUid)))
+			{
+				MeshNode->GetMaterialDependencies(MaterialDependencies);
+			}
+			const int32 MaterialCount = MaterialDependencies.Num();
+			for (int32 MaterialIndex = 0; MaterialIndex < MaterialCount; ++MaterialIndex)
+			{
+				const FString MaterialFactoryNodeUid = UInterchangeMaterialFactoryNode::GetMaterialFactoryNodeUidFromMaterialNodeUid(MaterialDependencies[MaterialIndex]);
+				if (BaseNodeContainer->IsNodeUidValid(MaterialFactoryNodeUid))
+				{
+					//Create a factory dependency so Material asset are import before the skeletal mesh asset
+					TArray<FString> FactoryDependencies;
+					SkeletalMeshFactoryNode->GetFactoryDependencies(FactoryDependencies);
+					if (!FactoryDependencies.Contains(MaterialFactoryNodeUid))
+					{
+						SkeletalMeshFactoryNode->AddFactoryDependencyUid(MaterialFactoryNodeUid);
+					}
+				}
+			}
+			LodDataNode->AddMeshUid(NodeUid);
+		}
+	}
+}
+
+void UInterchangeGenericAssetsPipeline::PostImportPhysicsAssetImport(UObject* CreatedAsset, UInterchangeBaseNode* Node)
+{
+#if WITH_EDITOR
+	if (!bCreatePhysicsAsset || !BaseNodeContainer)
+	{
+		return;
+	}
+
+	UPhysicsAsset* CreatedPhysicsAsset = Cast<UPhysicsAsset>(CreatedAsset);
+	if (!CreatedPhysicsAsset)
+	{
+		return;
+	}
+	if (UInterchangePhysicsAssetFactoryNode* PhysicsAssetFactoryNode = Cast<UInterchangePhysicsAssetFactoryNode>(Node))
+	{
+		FString SkeletalMeshFactoryNodeUid;
+		if (PhysicsAssetFactoryNode->GetCustomSkeletalMeshUid(SkeletalMeshFactoryNodeUid))
+		{
+			if (UInterchangeSkeletalMeshFactoryNode* SkeletalMeshFactoryNode = Cast<UInterchangeSkeletalMeshFactoryNode>(BaseNodeContainer->GetNode(SkeletalMeshFactoryNodeUid)))
+			{
+				if (SkeletalMeshFactoryNode->ReferenceObject.IsValid())
+				{
+					if (USkeletalMesh* SkeletalMesh = Cast<USkeletalMesh>(SkeletalMeshFactoryNode->ReferenceObject.TryLoad()))
+					{
+						auto CreateFromSkeletalMeshLambda = [CreatedPhysicsAsset, SkeletalMesh]()
+						{
+							FPhysAssetCreateParams NewBodyData;
+							FText CreationErrorMessage;
+							if (!FPhysicsAssetUtils::CreateFromSkeletalMesh(CreatedPhysicsAsset, SkeletalMesh, NewBodyData, CreationErrorMessage))
+							{
+								//TODO: Log an error
+							}
+						};
+
+						if (!IsInGameThread() && SkeletalMesh->IsCompiling())
+						{
+							//If the skeletalmesh is compiling we have to stall on the main thread
+							Async(EAsyncExecution::TaskGraphMainThread, [CreateFromSkeletalMeshLambda]()
+							{
+								CreateFromSkeletalMeshLambda();
+							});
+						}
+						else
+						{
+							CreateFromSkeletalMeshLambda();
+						}
+					}
+				}
+			}
+		}
+	}
+#endif //WITH_EDITOR
+}
+
+void UInterchangeGenericAssetsPipeline::ImplementUseSourceNameForAssetOptionSkeletalMesh(const int32 MeshesAndAnimsImportedNodeCount)
+{
+	const UClass* SkeletalMeshFactoryNodeClass = UInterchangeSkeletalMeshFactoryNode::StaticClass();
+	TArray<FString> SkeletalMeshNodeUids;
+	BaseNodeContainer->GetNodes(SkeletalMeshFactoryNodeClass, SkeletalMeshNodeUids);
+	if (SkeletalMeshNodeUids.Num() == 0)
+	{
+		return;
+	}
+	//If we import only one asset, and bUseSourceNameForAsset is true, we want to rename the asset using the file name.
+	const bool bShouldChangeAssetName = (bUseSourceNameForAsset && MeshesAndAnimsImportedNodeCount == 1);
+	const FString SkeletalMeshUid = SkeletalMeshNodeUids[0];
+	UInterchangeSkeletalMeshFactoryNode* SkeletalMeshNode = Cast<UInterchangeSkeletalMeshFactoryNode>(BaseNodeContainer->GetNode(SkeletalMeshUid));
+	if (!SkeletalMeshNode)
+	{
+		return;
+	}
+
+	FString DisplayLabelName = SkeletalMeshNode->GetDisplayLabel();
+		
+	if (bShouldChangeAssetName)
+	{
+		DisplayLabelName = FPaths::GetBaseFilename(SourceDatas[0]->GetFilename());
+		SkeletalMeshNode->SetDisplayLabel(DisplayLabelName);
+	}
+
+	//Also set the skeleton factory node name
+	TArray<FString> LodDataUids;
+	SkeletalMeshNode->GetLodDataUniqueIds(LodDataUids);
+	if (LodDataUids.Num() > 0)
+	{
+		//Get the skeleton from the base LOD, skeleton is shared with all LODs
+		if (UInterchangeSkeletalMeshLodDataNode* SkeletalMeshLodDataNode = Cast<UInterchangeSkeletalMeshLodDataNode>(BaseNodeContainer->GetNode(LodDataUids[0])))
+		{
+			//If the user did not specify any skeleton
+			if (Skeleton.IsNull())
+			{
+				FString SkeletalMeshSkeletonUid;
+				SkeletalMeshLodDataNode->GetCustomSkeletonUid(SkeletalMeshSkeletonUid);
+				UInterchangeSkeletonFactoryNode* SkeletonFactoryNode = Cast<UInterchangeSkeletonFactoryNode>(BaseNodeContainer->GetNode(SkeletalMeshSkeletonUid));
+				if (SkeletonFactoryNode)
+				{
+					const FString SkeletonName = DisplayLabelName + TEXT("_Skeleton");
+					SkeletonFactoryNode->SetDisplayLabel(SkeletonName);
+				}
+			}
+		}
+	}
+	const UClass* PhysicsAssetFactoryNodeClass = UInterchangePhysicsAssetFactoryNode::StaticClass();
+	TArray<FString> PhysicsAssetNodeUids;
+	BaseNodeContainer->GetNodes(PhysicsAssetFactoryNodeClass, PhysicsAssetNodeUids);
+	for (const FString& PhysicsAssetNodeUid : PhysicsAssetNodeUids)
+	{
+		UInterchangePhysicsAssetFactoryNode* PhysicsAssetFactoryNode = Cast<UInterchangePhysicsAssetFactoryNode>(BaseNodeContainer->GetNode(PhysicsAssetNodeUid));
+		if (!ensure(PhysicsAssetFactoryNode))
+		{
+			continue;
+		}
+		FString PhysicsAssetSkeletalMeshUid;
+		if (PhysicsAssetFactoryNode->GetCustomSkeletalMeshUid(PhysicsAssetSkeletalMeshUid) && PhysicsAssetSkeletalMeshUid.Equals(SkeletalMeshUid))
+		{
+			//Rename this asset
+			const FString PhysicsAssetName = DisplayLabelName + TEXT("_PhysicsAsset");
+			PhysicsAssetFactoryNode->SetDisplayLabel(PhysicsAssetName);
+		}
+	}
+}
\ No newline at end of file
diff --git a/Engine/Plugins/Runtime/Interchange/Source/Pipelines/Private/InterchangePipelineMeshesUtilities.cpp b/Engine/Plugins/Runtime/Interchange/Source/Pipelines/Private/InterchangePipelineMeshesUtilities.cpp
--- a/Engine/Plugins/Runtime/Interchange/Source/Pipelines/Private/InterchangePipelineMeshesUtilities.cpp
+++ b/Engine/Plugins/Runtime/Interchange/Source/Pipelines/Private/InterchangePipelineMeshesUtilities.cpp
@@ -21,99 +21,101 @@ UInterchangePipelineMeshesUtilities* UInterchangePipelineMeshesUtilities::Create
 	//Find all translated node we need for this pipeline
 	BaseNodeContainer->IterateNodes([&PipelineMeshesUtilities, &BaseNodeContainer, &SkeletonRootNodeUids](const FString& NodeUid, UInterchangeBaseNode* Node)
 	{
-		switch (Node->GetnodeContainerType())
+		if (Node->GetnodeContainerType() == EInterchangeNodeContainerType::NodeContainerType_TranslatedAsset)
 		{
-			case EInterchangeNodeContainerType::NodeContainerType_TranslatedAsset:
+			if (UInterchangeMeshNode* MeshNode = Cast<UInterchangeMeshNode>(Node))
 			{
-				if (UInterchangeMeshNode* MeshNode = Cast<UInterchangeMeshNode>(Node))
-				{
-					FInterchangeMeshGeometry& MeshGeometry = PipelineMeshesUtilities->MeshGeometriesPerMeshUid.FindOrAdd(NodeUid);
-					MeshGeometry.MeshUid = NodeUid;
-					MeshGeometry.MeshNode = MeshNode;
-				}
+				FInterchangeMeshGeometry& MeshGeometry = PipelineMeshesUtilities->MeshGeometriesPerMeshUid.FindOrAdd(NodeUid);
+				MeshGeometry.MeshUid = NodeUid;
+				MeshGeometry.MeshNode = MeshNode;
 			}
-			break;
-			case EInterchangeNodeContainerType::NodeContainerType_TranslatedScene:
+		}
+	});
+
+	//Find all translated scene node we need for this pipeline
+	BaseNodeContainer->IterateNodes([&PipelineMeshesUtilities, &BaseNodeContainer, &SkeletonRootNodeUids](const FString& NodeUid, UInterchangeBaseNode* Node)
+	{
+		if (Node->GetnodeContainerType() == EInterchangeNodeContainerType::NodeContainerType_TranslatedScene)
+		{
+			if (UInterchangeSceneNode* SceneNode = Cast<UInterchangeSceneNode>(Node))
 			{
-				if (UInterchangeSceneNode* SceneNode = Cast<UInterchangeSceneNode>(Node))
+				if (SceneNode->IsSpecializedTypeContains(UE::Interchange::FSceneNodeStaticData::GetJointSpecializeTypeString()))
 				{
-					if (SceneNode->IsSpecializedTypeContains(UE::Interchange::FSceneNodeStaticData::GetJointSpecializeTypeString()))
+					const UInterchangeSceneNode* ParentJointNode = Cast<UInterchangeSceneNode>(BaseNodeContainer->GetNode(SceneNode->GetParentUid()));
+					if (!ParentJointNode || !ParentJointNode->IsSpecializedTypeContains(UE::Interchange::FSceneNodeStaticData::GetJointSpecializeTypeString()))
 					{
-						const UInterchangeSceneNode* ParentJointNode = Cast<UInterchangeSceneNode>(BaseNodeContainer->GetNode(SceneNode->GetParentUid()));
-						if (!ParentJointNode || !ParentJointNode->IsSpecializedTypeContains(UE::Interchange::FSceneNodeStaticData::GetJointSpecializeTypeString()))
-						{
-							SkeletonRootNodeUids.Add(SceneNode->GetUniqueID());
-						}
+						SkeletonRootNodeUids.Add(SceneNode->GetUniqueID());
 					}
+				}
 
-					FString MeshUid;
-					if (SceneNode->GetCustomMeshDependencyUid(MeshUid))
+				FString MeshUid;
+				if (SceneNode->GetCustomMeshDependencyUid(MeshUid))
+				{
+					UInterchangeSceneNode* ParentMeshSceneNode = Cast<UInterchangeSceneNode>(BaseNodeContainer->GetNode(SceneNode->GetParentUid()));
+					if (ParentMeshSceneNode)
 					{
-						UInterchangeSceneNode* ParentMeshSceneNode = Cast<UInterchangeSceneNode>(BaseNodeContainer->GetNode(SceneNode->GetParentUid()));
-						if (ParentMeshSceneNode)
+						UInterchangeSceneNode* LodGroupNode = nullptr;
+						int32 LodIndex = 0;
+						FString LastChildUid = SceneNode->GetUniqueID();
+						do
 						{
-							UInterchangeSceneNode* LodGroupNode = nullptr;
-							int32 LodIndex = 0;
-							FString LastChildUid = SceneNode->GetUniqueID();
-							do
+							if (ParentMeshSceneNode->IsSpecializedTypeContains(UE::Interchange::FSceneNodeStaticData::GetLodGroupSpecializeTypeString()))
 							{
-								if (ParentMeshSceneNode->IsSpecializedTypeContains(UE::Interchange::FSceneNodeStaticData::GetLodGroupSpecializeTypeString()))
+								LodGroupNode = ParentMeshSceneNode;
+								TArray<FString> LodGroupChildrens = BaseNodeContainer->GetNodeChildrenUids(ParentMeshSceneNode->GetUniqueID());
+								for (int32 ChildLodIndex = 0; ChildLodIndex < LodGroupChildrens.Num(); ++ChildLodIndex)
 								{
-									LodGroupNode = ParentMeshSceneNode;
-									TArray<FString> LodGroupChildrens = BaseNodeContainer->GetNodeChildrenUids(ParentMeshSceneNode->GetUniqueID());
-									for (int32 ChildLodIndex = 0; ChildLodIndex < LodGroupChildrens.Num(); ++ChildLodIndex)
+									const FString& ChildrenUid = LodGroupChildrens[ChildLodIndex];
+									if (ChildrenUid.Equals(LastChildUid))
 									{
-										const FString& ChildrenUid = LodGroupChildrens[ChildLodIndex];
-										if (ChildrenUid.Equals(LastChildUid))
-										{
-											LodIndex = ChildLodIndex;
-											break;
-										}
+										LodIndex = ChildLodIndex;
+										break;
 									}
-									break;
 								}
-								LastChildUid = ParentMeshSceneNode->GetUniqueID();
-								ParentMeshSceneNode = Cast<UInterchangeSceneNode>(BaseNodeContainer->GetNode(ParentMeshSceneNode->GetParentUid()));
-							} while (ParentMeshSceneNode);
+								break;
+							}
+							LastChildUid = ParentMeshSceneNode->GetUniqueID();
+							ParentMeshSceneNode = Cast<UInterchangeSceneNode>(BaseNodeContainer->GetNode(ParentMeshSceneNode->GetParentUid()));
+						} while (ParentMeshSceneNode);
 
-							FInterchangeMeshGeometry& MeshGeometry = PipelineMeshesUtilities->MeshGeometriesPerMeshUid.FindChecked(MeshUid);
-							if (LodGroupNode)
+						FInterchangeMeshGeometry& MeshGeometry = PipelineMeshesUtilities->MeshGeometriesPerMeshUid.FindChecked(MeshUid);
+						if (LodGroupNode)
+						{
+							//We have a LOD
+							FInterchangeMeshInstance& MeshInstance = PipelineMeshesUtilities->MeshInstancesPerMeshInstanceUid.FindOrAdd(LodGroupNode->GetUniqueID());
+							if (MeshInstance.LodGroupNode != nullptr)
 							{
-								//We have a LOD
-								FInterchangeMeshInstance& MeshInstance = PipelineMeshesUtilities->MeshInstancesPerMeshInstanceUid.FindOrAdd(LodGroupNode->GetUniqueID());
-								if (MeshInstance.LodGroupNode != nullptr)
-								{
-									//This LodGroup was already created, verify everything is ok
-									checkSlow(MeshInstance.LodGroupNode == LodGroupNode);
-									checkSlow(MeshInstance.MeshInstanceUid.Equals(LodGroupNode->GetUniqueID()));
-								}
-								else
-								{
-									MeshInstance.LodGroupNode = LodGroupNode;
-									MeshInstance.MeshInstanceUid = LodGroupNode->GetUniqueID();
-								}
-								FInterchangeLodSceneNodeContainer& InstancedSceneNodes = MeshInstance.SceneNodePerLodIndex.FindOrAdd(LodIndex);
-								InstancedSceneNodes.SceneNodes.AddUnique(SceneNode);
-								MeshGeometry.ReferencingMeshInstanceUids.Add(MeshInstance.MeshInstanceUid);
-								MeshInstance.ReferencingMeshGeometryUids.Add(MeshUid);
-								MeshInstance.bReferenceSkinnedMesh |= MeshGeometry.MeshNode->IsSkinnedMesh();
+								//This LodGroup was already created, verify everything is ok
+								checkSlow(MeshInstance.LodGroupNode == LodGroupNode);
+								checkSlow(MeshInstance.MeshInstanceUid.Equals(LodGroupNode->GetUniqueID()));
 							}
 							else
 							{
-								FInterchangeMeshInstance& MeshInstance = PipelineMeshesUtilities->MeshInstancesPerMeshInstanceUid.FindOrAdd(NodeUid);
-								MeshInstance.LodGroupNode = nullptr;
-								MeshInstance.MeshInstanceUid = NodeUid;
-								FInterchangeLodSceneNodeContainer& InstancedSceneNodes = MeshInstance.SceneNodePerLodIndex.FindOrAdd(LodIndex);
-								InstancedSceneNodes.SceneNodes.AddUnique(SceneNode);
-								MeshGeometry.ReferencingMeshInstanceUids.Add(MeshInstance.MeshInstanceUid);
-								MeshInstance.ReferencingMeshGeometryUids.Add(MeshUid);
-								MeshInstance.bReferenceSkinnedMesh |= MeshGeometry.MeshNode->IsSkinnedMesh();
+								MeshInstance.LodGroupNode = LodGroupNode;
+								MeshInstance.MeshInstanceUid = LodGroupNode->GetUniqueID();
 							}
+							FInterchangeLodSceneNodeContainer& InstancedSceneNodes = MeshInstance.SceneNodePerLodIndex.FindOrAdd(LodIndex);
+							InstancedSceneNodes.SceneNodes.AddUnique(SceneNode);
+							MeshGeometry.ReferencingMeshInstanceUids.Add(MeshInstance.MeshInstanceUid);
+							MeshInstance.ReferencingMeshGeometryUids.Add(MeshUid);
+							MeshInstance.bReferenceSkinnedMesh |= MeshGeometry.MeshNode->IsSkinnedMesh();
+							MeshInstance.bReferenceBlendShape |= MeshGeometry.MeshNode->IsBlendShape();
+						}
+						else
+						{
+							FInterchangeMeshInstance& MeshInstance = PipelineMeshesUtilities->MeshInstancesPerMeshInstanceUid.FindOrAdd(NodeUid);
+							MeshInstance.LodGroupNode = nullptr;
+							MeshInstance.MeshInstanceUid = NodeUid;
+							FInterchangeLodSceneNodeContainer& InstancedSceneNodes = MeshInstance.SceneNodePerLodIndex.FindOrAdd(LodIndex);
+							InstancedSceneNodes.SceneNodes.AddUnique(SceneNode);
+							MeshGeometry.ReferencingMeshInstanceUids.Add(MeshInstance.MeshInstanceUid);
+							MeshInstance.ReferencingMeshGeometryUids.Add(MeshUid);
+							MeshInstance.bReferenceSkinnedMesh |= MeshGeometry.MeshNode->IsSkinnedMesh();
+							MeshInstance.bReferenceBlendShape |= MeshGeometry.MeshNode->IsBlendShape();
 						}
 					}
 				}
 			}
-			break;
 		}
 	});
 
@@ -196,7 +198,7 @@ void UInterchangePipelineMeshesUtilities::GetAllStaticMeshInstance(TArray<FStrin
 	for (const TPair<FString, FInterchangeMeshInstance>& MeshInstanceUidAndMeshInstance : MeshInstancesPerMeshInstanceUid)
 	{
 		const FInterchangeMeshInstance& MeshInstance = MeshInstanceUidAndMeshInstance.Value;
-		if (!MeshInstance.bReferenceSkinnedMesh)
+		if (!MeshInstance.bReferenceSkinnedMesh && !MeshInstance.bReferenceBlendShape)
 		{
 			MeshInstanceUids.Add(MeshInstance.MeshInstanceUid);
 		}
@@ -208,7 +210,7 @@ void UInterchangePipelineMeshesUtilities::IterateAllStaticMeshInstance(TFunction
 	for (const TPair<FString, FInterchangeMeshInstance>& MeshInstanceUidAndMeshInstance : MeshInstancesPerMeshInstanceUid)
 	{
 		const FInterchangeMeshInstance& MeshInstance = MeshInstanceUidAndMeshInstance.Value;
-		if (!MeshInstance.bReferenceSkinnedMesh)
+		if (!MeshInstance.bReferenceSkinnedMesh && !MeshInstance.bReferenceBlendShape)
 		{
 			IterationLambda(MeshInstance);
 		}
@@ -260,7 +262,7 @@ void UInterchangePipelineMeshesUtilities::GetAllStaticMeshGeometry(TArray<FStrin
 	for (const TPair<FString, FInterchangeMeshGeometry>& MeshGeometryUidAndMeshGeometry : MeshGeometriesPerMeshUid)
 	{
 		const FInterchangeMeshGeometry& MeshGeometry = MeshGeometryUidAndMeshGeometry.Value;
-		if (!MeshGeometry.MeshNode->IsSkinnedMesh())
+		if (!MeshGeometry.MeshNode->IsSkinnedMesh() && !MeshGeometry.MeshNode->IsBlendShape())
 		{
 			MeshGeometryUids.Add(MeshGeometry.MeshUid);
 		}
@@ -272,7 +274,7 @@ void UInterchangePipelineMeshesUtilities::IterateAllStaticMeshGeometry(TFunction
 	for (const TPair<FString, FInterchangeMeshGeometry>& MeshGeometryUidAndMeshGeometry : MeshGeometriesPerMeshUid)
 	{
 		const FInterchangeMeshGeometry& MeshGeometry = MeshGeometryUidAndMeshGeometry.Value;
-		if (!MeshGeometry.MeshNode->IsSkinnedMesh())
+		if (!MeshGeometry.MeshNode->IsSkinnedMesh() && !MeshGeometry.MeshNode->IsBlendShape())
 		{
 			IterationLambda(MeshGeometry);
 		}
diff --git a/Engine/Plugins/Runtime/Interchange/Source/Pipelines/Public/InterchangeGenericAssetsPipeline.h b/Engine/Plugins/Runtime/Interchange/Source/Pipelines/Public/InterchangeGenericAssetsPipeline.h
--- a/Engine/Plugins/Runtime/Interchange/Source/Pipelines/Public/InterchangeGenericAssetsPipeline.h
+++ b/Engine/Plugins/Runtime/Interchange/Source/Pipelines/Public/InterchangeGenericAssetsPipeline.h
@@ -33,6 +33,7 @@ template<class T> class TSubclassOf;
 #define COMMON_CATEGORY "Common"
 #define COMMON_MESHES_CATEGORY "Common Meshes"
 #define STATIC_MESHES_CATEGORY "Static Meshes"
+#define COMMON_SKELETAL_ANIMATIONS_CATEGORY "Common Skeletal Mesh and Animations"
 #define SKELETAL_MESHES_CATEGORY "Skeletal Meshes"
 #define ANIMATIONS_CATEGORY "Animations"
 #define MATERIALS_CATEGORY "Materials"
@@ -79,15 +80,15 @@ public:
 
 
 	/** Allow to convert mesh to a particular type */
-	UPROPERTY(EditAnywhere, BlueprintReadWrite, Category = COMMON_MESHES_CATEGORY)
-	TEnumAsByte<enum EForceMeshType> ForceAllMeshHasType = EForceMeshType::FMT_None;
+// 	UPROPERTY(EditAnywhere, BlueprintReadWrite, Category = COMMON_MESHES_CATEGORY)
+// 	TEnumAsByte<enum EForceMeshType> ForceAllMeshHasType = EForceMeshType::FMT_None;
 
 	/** If enable, meshes LODs will be imported. Note that it required the advanced bBakeMesh property to be enabled. */
 	UPROPERTY(EditAnywhere, BlueprintReadWrite, Category = COMMON_MESHES_CATEGORY, meta = (editcondition = "bBakeMeshes"))
 	bool bImportLods = true;
 
-	/** If enable, meshes will be baked with the hierarchy transform. */
-	UPROPERTY(EditAnywhere, BlueprintReadWrite, AdvancedDisplay, Category = COMMON_MESHES_CATEGORY)
+	/** If enable, meshes will be baked with the scene instance hierarchy transform. */
+	UPROPERTY(EditAnywhere, BlueprintReadWrite, Category = COMMON_MESHES_CATEGORY)
 	bool bBakeMeshes = true;
 
 	
@@ -103,6 +104,14 @@ public:
 	bool bCombineStaticMeshes = false;
 
 	
+	//////  COMMON_SKELETAL_ANIMATIONS_CATEGORY //////
+
+
+	/** Skeleton to use for imported asset. When importing a skeletal mesh, leaving this as "None" will create a new skeleton. When importing an animation this MUST be specified to import the asset. */
+ 	UPROPERTY(EditAnywhere, BlueprintReadWrite, Category = COMMON_SKELETAL_ANIMATIONS_CATEGORY)
+ 	TObjectPtr<class USkeleton> Skeleton;
+
+	
 	//////	SKELETAL_MESHES_CATEGORY Properties //////
 
 	
@@ -110,17 +119,28 @@ public:
 	UPROPERTY(EditAnywhere, BlueprintReadWrite, Category = SKELETAL_MESHES_CATEGORY)
 	bool bImportSkeletalMeshes = true;
 
-	/** If enable all translated skinned mesh node will be imported has a one skeletal mesh, note that it can still create several skeletal mesh for each different root joint. */
+	/** If enable all translated skinned mesh node will be imported has a one skeletal mesh, note that it can still create several skeletal mesh for each different skeleton root joint. */
 	UPROPERTY(EditAnywhere, BlueprintReadWrite, Category = SKELETAL_MESHES_CATEGORY)
 	bool bCombineSkeletalMeshes = true;
 
+	/** If enable any morph target shape will be imported. */
+	UPROPERTY(EditAnywhere, BlueprintReadWrite, Category = SKELETAL_MESHES_CATEGORY)
+	bool bImportMorphTargets = true;
+
+	/** If checked, create new PhysicsAsset if it doesn't have it */
+	UPROPERTY(EditAnywhere, BlueprintReadWrite, Category = SKELETAL_MESHES_CATEGORY)
+	uint32 bCreatePhysicsAsset : 1;
+
+	/** If this is set, use this specified PhysicsAsset. If its not set and bCreatePhysicsAsset is false, the importer will not generate or set any physic asset. */
+	UPROPERTY(EditAnywhere, BlueprintReadWrite, Category = SKELETAL_MESHES_CATEGORY, meta = (editcondition = "!bCreatePhysicsAsset"))
+	TObjectPtr<class UPhysicsAsset> PhysicsAsset;
 	
 	//////	ANIMATIONS_CATEGORY Properties //////
 
 
 	/** If enable, import the animation asset find in the sources. */
-	UPROPERTY(EditAnywhere, BlueprintReadWrite, Category = ANIMATIONS_CATEGORY)
-	bool bImportAnimations = true;
+// 	UPROPERTY(EditAnywhere, BlueprintReadWrite, Category = ANIMATIONS_CATEGORY)
+// 	bool bImportAnimations = true;
 
 
 	//////	MATERIALS_CATEGORY Properties //////
@@ -142,11 +162,13 @@ public:
 	// END Pre import pipeline properties
 	//////////////////////////////////////////////////////////////////////////
 
+	virtual void PreDialogCleanup(const FName PipelineStackName) override;
+
 protected:
 
 	virtual bool ExecutePreImportPipeline(UInterchangeBaseNodeContainer* InBaseNodeContainer, const TArray<UInterchangeSourceData*>& InSourceDatas) override;
 
-	//virtual bool ExecutePostImportPipeline(const UInterchangeBaseNodeContainer* BaseNodeContainer, const FName& NodeKey, UObject* CreatedAsset) override;
+	virtual bool ExecutePostImportPipeline(const UInterchangeBaseNodeContainer* InBaseNodeContainer, const FString& NodeKey, UObject* CreatedAsset) override;
 
 	virtual bool CanExecuteOnAnyThread(EInterchangePipelineTask PipelineTask) override
 	{
@@ -179,9 +201,9 @@ private:
 	UInterchangeMaterialFactoryNode* CreateMaterialFactoryNode(const UInterchangeMaterialNode* MaterialNode);
 
 	
-	/** Mesh translated assets nodes */
-	TArray<UInterchangeMeshNode*> MeshNodes;
-	
+	/************************************************************************/
+	/* Skeletal mesh API BEGIN                                              */
+
 	/** Skeleton factory assets nodes */
 	TArray<UInterchangeSkeletonFactoryNode*> SkeletonFactoryNodes;
 
@@ -206,6 +228,19 @@ private:
 	 */
 	void AddLodDataToSkeletalMesh(const UInterchangeSkeletonFactoryNode* SkeletonFactoryNode, UInterchangeSkeletalMeshFactoryNode* SkeletalMeshFactoryNode, const TMap<int32, TArray<FString>>& NodeUidsPerLodIndex);
 
+	/**
+	 * This function will create any skeletalmesh we need to create according to the pipeline options
+	 */
+	bool ExecutePreImportPipelineSkeletalMesh();
+
+	/**
+	 * This function will finish creating the physics asset with the skeletalmesh render data
+	 */
+	void PostImportPhysicsAssetImport(UObject* CreatedAsset, UInterchangeBaseNode* Node);
+
+	/* Skeletal mesh API END                                                */
+	/************************************************************************/
+
 
 	/** Static mesh factory assets nodes */
 	TArray<UInterchangeStaticMeshFactoryNode*> StaticMeshFactoryNodes;
@@ -230,11 +265,17 @@ private:
 	 */
 	bool MakeMeshFactoryNodeUidAndDisplayLabel(const TMap<int32, TArray<FString>>& MeshUidsPerLodIndex, int32 LodIndex, FString& NewMeshUid, FString& DisplayLabel);
 
-	/** Translated scene nodes */
-	TArray<UInterchangeSceneNode*> SceneNodes;
 
 	/* Meshes utilities, to parse the translated graph and extract the meshes informations. */
 	UInterchangePipelineMeshesUtilities* PipelineMeshesUtilities;
+
+	/**
+	 * Implement pipeline option bUseSourceNameForAsset
+	 */
+	void ImplementUseSourceNameForAssetOption();
+
+	/** Specialize for skeletalmesh */
+	void ImplementUseSourceNameForAssetOptionSkeletalMesh(const int32 MeshesAndAnimsImportedNodeCount);
 };
 
 
diff --git a/Engine/Plugins/Runtime/Interchange/Source/Pipelines/Public/InterchangePipelineMeshesUtilities.h b/Engine/Plugins/Runtime/Interchange/Source/Pipelines/Public/InterchangePipelineMeshesUtilities.h
--- a/Engine/Plugins/Runtime/Interchange/Source/Pipelines/Public/InterchangePipelineMeshesUtilities.h
+++ b/Engine/Plugins/Runtime/Interchange/Source/Pipelines/Public/InterchangePipelineMeshesUtilities.h
@@ -41,6 +41,7 @@ struct FInterchangeMeshInstance
 	{
 		LodGroupNode = nullptr;
 		bReferenceSkinnedMesh = false;
+		bReferenceBlendShape = false;
 	}
 	/**
 	* This ID represent either 1: a lod group scene node uid or 2: a mesh scene node uid.
@@ -57,6 +58,9 @@ struct FInterchangeMeshInstance
 	UPROPERTY(EditAnywhere, Category = "Interchange | Pipeline | MeshesInstance")
 	bool bReferenceSkinnedMesh;
 
+	UPROPERTY(EditAnywhere, Category = "Interchange | Pipeline | MeshesInstance")
+	bool bReferenceBlendShape;
+
 	/**
 	* Each scene node here represent a mesh scene node. Only if we represent a lod group we can have more then 1 mesh scene node for a specific lod index.
 	*/
diff --git a/Engine/Source/Runtime/Interchange/Core/Private/InterchangePipelineBase.cpp b/Engine/Source/Runtime/Interchange/Core/Private/InterchangePipelineBase.cpp
--- a/Engine/Source/Runtime/Interchange/Core/Private/InterchangePipelineBase.cpp
+++ b/Engine/Source/Runtime/Interchange/Core/Private/InterchangePipelineBase.cpp
@@ -29,10 +29,12 @@ void UInterchangePipelineBase::LoadSettings(const FName PipelineStackName)
 	UClass* Class = this->GetClass();
 	for (FProperty* Property = Class->PropertyLink; Property; Property = Property->PropertyLinkNext)
 	{
-// 		if (!Property->HasAnyPropertyFlags(CPF_Config))
-// 		{
-// 			continue;
-// 		}
+		//Do not load a transient property
+		if (Property->HasAnyPropertyFlags(CPF_Transient))
+		{
+			continue;
+		}
+
 		FString Section = UE::Interchange::PipelinePrivate::CreateConfigSectionName(PipelineStackName, Class);
 		FString Key = Property->GetName();
 
@@ -124,10 +126,12 @@ void UInterchangePipelineBase::SaveSettings(const FName PipelineStackName)
 	UClass* Class = this->GetClass();
 	for (FProperty* Property = Class->PropertyLink; Property; Property = Property->PropertyLinkNext)
 	{
-// 		if (!Property->HasAnyPropertyFlags(CPF_Config))
-// 		{
-// 			continue;
-// 		}
+		//Do not save a transient property
+		if (Property->HasAnyPropertyFlags(CPF_Transient))
+		{
+			continue;
+		}
+
 		FString Section = UE::Interchange::PipelinePrivate::CreateConfigSectionName(PipelineStackName, Class);
 		FString Key = Property->GetName();
 
diff --git a/Engine/Source/Runtime/Interchange/Core/Private/Tests/StorageTest.cpp b/Engine/Source/Runtime/Interchange/Core/Private/Tests/StorageTest.cpp
--- a/Engine/Source/Runtime/Interchange/Core/Private/Tests/StorageTest.cpp
+++ b/Engine/Source/Runtime/Interchange/Core/Private/Tests/StorageTest.cpp
@@ -98,13 +98,30 @@ bool FAttributeStorageTest::RunTest(const FString& Parameters)
 		FAttributeKey Key = CreateUniqueKey();
 		if (!IsAttributeStorageResultSuccess(TestStorage.RegisterAttribute(Key, RefValue)))
 		{
-			AddError(FString(TEXT("`AttributeStorage` must handle adding FAttributeKey attribute")));
+			AddError(FString(TEXT("`AttributeStorage` must handle adding FName attribute")));
 		}
 		FName StoredValue;
 		TestStorage.GetAttributeHandle<FName>(Key).Get(StoredValue);
 		if(StoredValue != RefValue)
 		{
-			TestEqual(TEXT("`AttributeStorage` must handle add and retrieve FAttributeKey attribute"), StoredValue.ToString(), RefValue.ToString());
+			TestEqual(TEXT("`AttributeStorage` must handle add and retrieve FName attribute"), StoredValue.ToString(), RefValue.ToString());
+		}
+	}
+
+	//Adding FSoftObjectPath
+	{
+
+		const FSoftObjectPath RefValue(UClass::StaticClass());
+		FAttributeKey Key = CreateUniqueKey();
+		if (!IsAttributeStorageResultSuccess(TestStorage.RegisterAttribute(Key, RefValue)))
+		{
+			AddError(FString(TEXT("`AttributeStorage` must handle adding FSoftObjectPath attribute")));
+		}
+		FSoftObjectPath StoredValue;
+		TestStorage.GetAttributeHandle<FSoftObjectPath>(Key).Get(StoredValue);
+		if (StoredValue != RefValue)
+		{
+			TestEqual(TEXT("`AttributeStorage` must handle add and retrieve FSoftObjectPath attribute"), StoredValue.ToString(), RefValue.ToString());
 		}
 	}
 
diff --git a/Engine/Source/Runtime/Interchange/Core/Private/Types/AttributeStorage.cpp b/Engine/Source/Runtime/Interchange/Core/Private/Types/AttributeStorage.cpp
--- a/Engine/Source/Runtime/Interchange/Core/Private/Types/AttributeStorage.cpp
+++ b/Engine/Source/Runtime/Interchange/Core/Private/Types/AttributeStorage.cpp
@@ -64,6 +64,7 @@ namespace UE
 				case EAttributeTypes::FrameNumber: { AttributeTypeString = TEXT("FrameNumber"); } break;
 				case EAttributeTypes::FrameRate: { AttributeTypeString = TEXT("FrameRate"); } break;
 				case EAttributeTypes::FrameTime: { AttributeTypeString = TEXT("FrameTime"); } break;
+				case EAttributeTypes::SoftObjectPath: { AttributeTypeString = TEXT("SoftObjectPath"); } break;
 				default:
 				{
 					//Ensure if we ask an unknown type
@@ -469,6 +470,52 @@ namespace UE
 			return EAttributeStorageResult::Operation_Success;
 		}
 
+		EAttributeStorageResult FAttributeStorage::GetAttribute(const FAttributeKey& ElementAttributeKey, FSoftObjectPath& OutValue, TSpecializeType<FSoftObjectPath >) const
+		{
+			static_assert(TAttributeTypeTraits<FSoftObjectPath>::GetType() != EAttributeTypes::None, "Not a supported type for the attributes. Check EAttributeTypes for supported types");
+
+			//Lock the storage
+			FScopeLock ScopeLock(&StorageMutex);
+
+			const EAttributeTypes ValueType = TAttributeTypeTraits<FSoftObjectPath>::GetType();
+
+
+			const FAttributeAllocationInfo* AttributeAllocationInfo = AttributeAllocationTable.Find(ElementAttributeKey);
+			if (AttributeAllocationInfo)
+			{
+				if (AttributeAllocationInfo->Type != ValueType)
+				{
+					return EAttributeStorageResult::Operation_Error_WrongType;
+				}
+				if (!AttributeStorage.IsValidIndex(AttributeAllocationInfo->Offset))
+				{
+					return EAttributeStorageResult::Operation_Error_AttributeAllocationCorrupted;
+				}
+			}
+			else
+			{
+				//The key do not exist
+				return EAttributeStorageResult::Operation_Error_CannotFoundKey;
+			}
+
+
+			if (AttributeAllocationInfo->Size == 0)
+			{
+				OutValue = TEXT("");
+				return EAttributeStorageResult::Operation_Success;
+			}
+
+			FString ValueStr;
+
+			//Share the code with FString from here
+			ExtractFStringAttributeFromStorage(AttributeStorage.GetData(), AttributeAllocationInfo, ValueStr);
+
+			//Create the FSoftObjectPath and copy it to OutValue
+			OutValue = FSoftObjectPath(ValueStr);
+
+			return EAttributeStorageResult::Operation_Success;
+		}
+
 		void FAttributeStorage::DefragInternal()
 		{
 			//Defrag if the fragmented memory cost is bigger then the DefragRatio.
diff --git a/Engine/Source/Runtime/Interchange/Core/Public/InterchangeFactoryBase.h b/Engine/Source/Runtime/Interchange/Core/Public/InterchangeFactoryBase.h
--- a/Engine/Source/Runtime/Interchange/Core/Public/InterchangeFactoryBase.h
+++ b/Engine/Source/Runtime/Interchange/Core/Public/InterchangeFactoryBase.h
@@ -109,7 +109,7 @@ public:
 	/**
 	 * Parameters to pass to CreateAsset function
 	 */
-	struct FPostImportGameThreadCallbackParams
+	struct FImportPreCompletedCallbackParams
 	{
 		/** The source data, mainly use to set the asset import data file. TODO: we have to refactor UAssetImportData, the source data should be the base class for this now */
 		const UInterchangeSourceData* SourceData = nullptr;
@@ -121,8 +121,22 @@ public:
 		TArray<UInterchangePipelineBase*> Pipelines;
 ;
 	};
-	/* This function is call in the completion task on the main thread, use it to call main thread post creation step for your assets*/
-	virtual void PostImportGameThreadCallback(const FPostImportGameThreadCallbackParams& Arguments) const
+
+	/*
+	 * This function is call in the pre completion task on the main thread, use it to call main thread post creation step for your assets
+	 * @note - This function is called when starting the pre completion task (before PostEditChange is called for the asset).
+	 */
+	virtual void PreImportPreCompletedCallback(const FImportPreCompletedCallbackParams& Arguments) const
+	{
+		check(IsInGameThread());
+		return;
+	}
+	
+	/*
+	 * This function is call in the pre completion task on the main thread, use it to call main thread post creation step for your assets
+	 * @note - This function is called at the end of the pre completion task (after PostEditChange is called for the asset).
+	 */
+	virtual void PostImportPreCompletedCallback(const FImportPreCompletedCallbackParams& Arguments) const
 	{
 		check(IsInGameThread());
 		return;
diff --git a/Engine/Source/Runtime/Interchange/Core/Public/InterchangePipelineBase.h b/Engine/Source/Runtime/Interchange/Core/Public/InterchangePipelineBase.h
--- a/Engine/Source/Runtime/Interchange/Core/Public/InterchangePipelineBase.h
+++ b/Engine/Source/Runtime/Interchange/Core/Public/InterchangePipelineBase.h
@@ -28,7 +28,9 @@ public:
 
 	/**
 	 * Non virtual helper to allow blueprint to implement event base function to implement a pre import pipeline,
-	 * the Interchange manager is calling this function not the virtual one that is call by the default implementation.
+	 * It is call after the translation and before we parse the graph to call the factory. This is where factory node should be created
+	 * by the pipeline. Each factory node should be send to a a interchange factory to create an unreal asset.
+	 * @note - the Interchange manager is calling this function not the virtual one that is call by the default implementation.
 	 */
 	UFUNCTION(BlueprintNativeEvent, BlueprintCallable, Category = "Interchange | Translator")
 	bool ScriptedExecutePreImportPipeline(UInterchangeBaseNodeContainer* BaseNodeContainer, const TArray<UInterchangeSourceData*>& SourceDatas);
@@ -41,12 +43,15 @@ public:
 
 	/**
 	 * Non virtual helper to allow blueprint to implement event base function to implement a post import pipeline,
-	 * the Interchange manager is calling this function not the virtual one that is call by the default implementation.
+	 * It is call after we completely import an asset. PostEditChange is already called. Some assets uses asynchronous build.
+	 * This can be useful if you need build data of an asset to finish the setup of another asset.
+	 * @example - PhysicsAsset need skeletal mesh render data to be build properly.
+	 * @note - the Interchange manager is calling this function not the virtual one that is call by the default implementation.
 	 */
 	UFUNCTION(BlueprintNativeEvent, BlueprintCallable, Category = "Interchange | Translator")
-	bool ScriptedExecutePostImportPipeline(const UInterchangeBaseNodeContainer* BaseNodeContainer, const FName& NodeKey, UObject* CreatedAsset);
+	bool ScriptedExecutePostImportPipeline(const UInterchangeBaseNodeContainer* BaseNodeContainer, const FString& NodeKey, UObject* CreatedAsset);
 	/** The default implementation (call if the blueprint do not have any implementation) will call the virtual ExecuteImportPipeline */
-	bool ScriptedExecutePostImportPipeline_Implementation(const UInterchangeBaseNodeContainer* BaseNodeContainer, const FName& NodeKey, UObject* CreatedAsset)
+	bool ScriptedExecutePostImportPipeline_Implementation(const UInterchangeBaseNodeContainer* BaseNodeContainer, const FString& NodeKey, UObject* CreatedAsset)
 	{
 		//By default we call the virtual import pipeline execution
 		return ExecutePostImportPipeline(BaseNodeContainer, NodeKey, CreatedAsset);
@@ -94,6 +99,8 @@ public:
 
 	void SaveSettings(const FName PipelineStackName);
 
+	virtual void PreDialogCleanup(const FName PipelineStackName) {};
+	
 	/**
 	 * This function is used to add the given message object directly into the results for this operation.
 	 */
@@ -134,7 +141,7 @@ protected:
 	 * The interchange manager is not calling this function directly. It is calling the blueprint native event in case this object is a blueprint derive object.
 	 * By default the scripted implementation is calling this virtual pipeline.
 	 */
-	virtual bool ExecutePostImportPipeline(const UInterchangeBaseNodeContainer* BaseNodeContainer, const FName& NodeKey, UObject* CreatedAsset)
+	virtual bool ExecutePostImportPipeline(const UInterchangeBaseNodeContainer* BaseNodeContainer, const FString& NodeKey, UObject* CreatedAsset)
 	{
 		return false;
 	}
diff --git a/Engine/Source/Runtime/Interchange/Core/Public/Types/AttributeStorage.h b/Engine/Source/Runtime/Interchange/Core/Public/Types/AttributeStorage.h
--- a/Engine/Source/Runtime/Interchange/Core/Public/Types/AttributeStorage.h
+++ b/Engine/Source/Runtime/Interchange/Core/Public/Types/AttributeStorage.h
@@ -10,6 +10,7 @@
 #include "Misc/FrameRate.h"
 #include "Misc/NetworkGuid.h"
 #include "Misc/ScopeLock.h"
+#include "UObject/SoftObjectPath.h"
 
 //Interchange namespace
 namespace UE
@@ -192,6 +193,7 @@ namespace UE
 			FrameNumber			= 45,
 			FrameRate			= 46,
 			FrameTime			= 47,
+			SoftObjectPath		= 48,
 		};
 
 		/**
@@ -772,6 +774,16 @@ namespace UE
 				});
 			}
 
+			EAttributeStorageResult SetAttribute(FAttributeAllocationInfo* AttributeAllocationInfo, const FSoftObjectPath& Value, TSpecializeType<FSoftObjectPath >)
+			{
+				//FSoftObjectPath can be hold has a FString and restore from it
+				FString ValueStr = Value.ToString();
+				return MultiSizeSetAttribute(AttributeAllocationInfo, Value, [&ValueStr]()->uint8*
+				{
+					return (uint8*)(*ValueStr);
+				});
+			}
+
 			template<typename T>
 			EAttributeStorageResult GetAttribute(const FAttributeKey& ElementAttributeKey, T& OutValue, TSpecializeType<T>) const
 			{
@@ -826,6 +838,8 @@ namespace UE
 
 			EAttributeStorageResult GetAttribute(const FAttributeKey& ElementAttributeKey, FName& OutValue, TSpecializeType<FName >) const;
 
+			EAttributeStorageResult GetAttribute(const FAttributeKey& ElementAttributeKey, FSoftObjectPath& OutValue, TSpecializeType<FSoftObjectPath >) const;
+
 			template<typename T>
 			static uint64 GetValueSize(const T& Value, TSpecializeType<T>)
 			{
@@ -845,6 +859,12 @@ namespace UE
 				return GetValueSize(ValueStr, TSpecializeType<FString >());
 			}
 
+			static uint64 GetValueSize(const FSoftObjectPath& Value, TSpecializeType<FSoftObjectPath>)
+			{
+				FString ValueStr = Value.ToString();
+				return GetValueSize(ValueStr, TSpecializeType<FString >());
+			}
+
 			static uint64 GetValueSize(const TArray<uint8>& Value, TSpecializeType<TArray<uint8> >);
 
 			static uint64 GetValueSize(const TArray64<uint8>& Value, TSpecializeType<TArray64<uint8> >);
@@ -1313,5 +1333,11 @@ namespace UE
 			static CONSTEXPR EAttributeTypes GetType() { return EAttributeTypes::FrameTime; }
 		};
 
+		/** Implements variant type traits for the built-in SoftObjectPath type. */
+		template<> struct TAttributeTypeTraits<FSoftObjectPath>
+		{
+			static CONSTEXPR EAttributeTypes GetType() { return EAttributeTypes::SoftObjectPath; }
+		};
+
 	} //ns interchange
 }
diff --git a/Engine/Source/Runtime/Interchange/Engine/Private/InterchangeManager.cpp b/Engine/Source/Runtime/Interchange/Engine/Private/InterchangeManager.cpp
--- a/Engine/Source/Runtime/Interchange/Engine/Private/InterchangeManager.cpp
+++ b/Engine/Source/Runtime/Interchange/Engine/Private/InterchangeManager.cpp
@@ -135,6 +135,10 @@ void UE::Interchange::FImportAsyncHelper::CancelAndWaitUntilDoneSynchronously()
 	{
 		TasksToComplete.Add(PreAsyncCompletionTask);
 	}
+	if (PreCompletionTask.GetReference())
+	{
+		TasksToComplete.Add(PreCompletionTask);
+	}
 	if (CompletionTask.GetReference())
 	{
 		//Completion task will make sure any created asset before canceling will be mark for delete
diff --git a/Engine/Source/Runtime/Interchange/Engine/Private/Tasks/InterchangeTaskCompletion.cpp b/Engine/Source/Runtime/Interchange/Engine/Private/Tasks/InterchangeTaskCompletion.cpp
--- a/Engine/Source/Runtime/Interchange/Engine/Private/Tasks/InterchangeTaskCompletion.cpp
+++ b/Engine/Source/Runtime/Interchange/Engine/Private/Tasks/InterchangeTaskCompletion.cpp
@@ -26,20 +26,17 @@ void UE::Interchange::FTaskPreAsyncCompletion::DoTask(ENamedThreads::Type Curren
 	AsyncHelper->ReleaseTranslatorsSource();
 }
 
-void UE::Interchange::FTaskCompletion::DoTask(ENamedThreads::Type CurrentThread, const FGraphEventRef& MyCompletionGraphEvent)
+void UE::Interchange::FTaskPreCompletion::DoTask(ENamedThreads::Type CurrentThread, const FGraphEventRef& MyCompletionGraphEvent)
 {
 #if INTERCHANGE_TRACE_ASYNCHRONOUS_TASK_ENABLED
-	INTERCHANGE_TRACE_ASYNCHRONOUS_TASK(Completion)
+	INTERCHANGE_TRACE_ASYNCHRONOUS_TASK(PreCompletion)
 #endif
 	TSharedPtr<FImportAsyncHelper, ESPMode::ThreadSafe> AsyncHelper = WeakAsyncHelper.Pin();
 	check(AsyncHelper.IsValid());
 
-	//No need anymore of the translators sources
-	AsyncHelper->ReleaseTranslatorsSource();
-
 	UInterchangeResultsContainer* Results = AsyncHelper->AssetImportResult->GetResults();
 
-	for(TPair<int32, TArray<FImportAsyncHelper::FImportedAssetInfo>>& AssetInfosPerSourceIndexPair : AsyncHelper->ImportedAssetsPerSourceIndex)
+	for (TPair<int32, TArray<FImportAsyncHelper::FImportedAssetInfo>>& AssetInfosPerSourceIndexPair : AsyncHelper->ImportedAssetsPerSourceIndex)
 	{
 		//Verify if the task was cancel
 		if (AsyncHelper->bCancel)
@@ -54,15 +51,15 @@ void UE::Interchange::FTaskCompletion::DoTask(ENamedThreads::Type CurrentThread,
 		{
 			UObject* Asset = AssetInfo.ImportAsset;
 			//In case Some factory code cannot run outside of the main thread we offer this callback to finish the work before calling post edit change (building the asset)
-			if(bCallPostImportGameThreadCallback && AssetInfo.Factory)
+			if (bCallPostImportGameThreadCallback && AssetInfo.Factory)
 			{
-				UInterchangeFactoryBase::FPostImportGameThreadCallbackParams Arguments;
+				UInterchangeFactoryBase::FImportPreCompletedCallbackParams Arguments;
 				Arguments.ImportedObject = Asset;
 				Arguments.SourceData = AsyncHelper->SourceDatas[SourceIndex];
 				Arguments.NodeUniqueID = AssetInfo.NodeUniqueId;
 				Arguments.NodeContainer = AsyncHelper->BaseNodeContainers[SourceIndex].Get();
 				Arguments.Pipelines = AsyncHelper->Pipelines;
-				AssetInfo.Factory->PostImportGameThreadCallback(Arguments);
+				AssetInfo.Factory->PreImportPreCompletedCallback(Arguments);
 			}
 
 			if (Asset == nullptr)
@@ -85,6 +82,50 @@ void UE::Interchange::FTaskCompletion::DoTask(ENamedThreads::Type CurrentThread,
 			Asset->PostEditChange();
 #endif
 			//Post import broadcast
+			if (!AsyncHelper->TaskData.ReimportObject)
+			{
+				//Notify the asset registry, only when we have created the asset
+				FAssetRegistryModule::AssetCreated(Asset);
+			}
+			AsyncHelper->AssetImportResult->AddImportedAsset(Asset);
+			//In case Some factory code cannot run outside of the main thread we offer this callback to finish the work after calling post edit change (building the asset)
+			//Its possible the build of the asset to be asynchronous, the factory must handle is own asset correctly
+			if (bCallPostImportGameThreadCallback && AssetInfo.Factory)
+			{
+				UInterchangeFactoryBase::FImportPreCompletedCallbackParams Arguments;
+				Arguments.ImportedObject = Asset;
+				Arguments.SourceData = AsyncHelper->SourceDatas[SourceIndex];
+				Arguments.NodeUniqueID = AssetInfo.NodeUniqueId;
+				Arguments.NodeContainer = AsyncHelper->BaseNodeContainers[SourceIndex].Get();
+				Arguments.Pipelines = AsyncHelper->Pipelines;
+				AssetInfo.Factory->PostImportPreCompletedCallback(Arguments);
+			}
+		}
+	}
+}
+
+
+void UE::Interchange::FTaskCompletion::DoTask(ENamedThreads::Type CurrentThread, const FGraphEventRef& MyCompletionGraphEvent)
+{
+#if INTERCHANGE_TRACE_ASYNCHRONOUS_TASK_ENABLED
+	INTERCHANGE_TRACE_ASYNCHRONOUS_TASK(Completion)
+#endif
+	TSharedPtr<FImportAsyncHelper, ESPMode::ThreadSafe> AsyncHelper = WeakAsyncHelper.Pin();
+	check(AsyncHelper.IsValid());
+
+	//No need anymore of the translators sources
+	AsyncHelper->ReleaseTranslatorsSource();
+	for(TPair<int32, TArray<FImportAsyncHelper::FImportedAssetInfo>>& AssetInfosPerSourceIndexPair : AsyncHelper->ImportedAssetsPerSourceIndex)
+	{
+		//Verify if the task was cancel
+		if (AsyncHelper->bCancel)
+		{
+			break;
+		}
+		const int32 SourceIndex = AssetInfosPerSourceIndexPair.Key;
+		for (const FImportAsyncHelper::FImportedAssetInfo& AssetInfo : AssetInfosPerSourceIndexPair.Value)
+		{
+			UObject* Asset = AssetInfo.ImportAsset;
 			if (AsyncHelper->TaskData.ReimportObject)
 			{
 				InterchangeManager->OnAssetPostReimport.Broadcast(Asset);
@@ -92,11 +133,7 @@ void UE::Interchange::FTaskCompletion::DoTask(ENamedThreads::Type CurrentThread,
 			else
 			{
 				InterchangeManager->OnAssetPostImport.Broadcast(Asset);
-				//Notify the asset registry, only when we have created the asset
-				FAssetRegistryModule::AssetCreated(Asset);
 			}
-
-			AsyncHelper->AssetImportResult->AddImportedAsset(Asset);
 		}
 	}
 
diff --git a/Engine/Source/Runtime/Interchange/Engine/Private/Tasks/InterchangeTaskCompletion.h b/Engine/Source/Runtime/Interchange/Engine/Private/Tasks/InterchangeTaskCompletion.h
--- a/Engine/Source/Runtime/Interchange/Engine/Private/Tasks/InterchangeTaskCompletion.h
+++ b/Engine/Source/Runtime/Interchange/Engine/Private/Tasks/InterchangeTaskCompletion.h
@@ -41,6 +41,35 @@ namespace UE
 			void DoTask(ENamedThreads::Type CurrentThread, const FGraphEventRef& MyCompletionGraphEvent);
 		};
 
+		class FTaskPreCompletion
+		{
+		private:
+			UInterchangeManager* InterchangeManager;
+			TWeakPtr<FImportAsyncHelper, ESPMode::ThreadSafe> WeakAsyncHelper;
+		public:
+			FTaskPreCompletion(UInterchangeManager* InInterchangeManager, TWeakPtr<FImportAsyncHelper, ESPMode::ThreadSafe> InAsyncHelper)
+				: InterchangeManager(InInterchangeManager)
+				, WeakAsyncHelper(InAsyncHelper)
+			{
+			}
+
+			static FORCEINLINE ENamedThreads::Type GetDesiredThread()
+			{
+				return ENamedThreads::GameThread;
+			}
+			static FORCEINLINE ESubsequentsMode::Type GetSubsequentsMode()
+			{
+				return ESubsequentsMode::TrackSubsequents;
+			}
+
+			FORCEINLINE TStatId GetStatId() const
+			{
+				RETURN_QUICK_DECLARE_CYCLE_STAT(FTaskPreCompletion, STATGROUP_TaskGraphTasks);
+			}
+
+			void DoTask(ENamedThreads::Type CurrentThread, const FGraphEventRef& MyCompletionGraphEvent);
+		};
+
 		class FTaskCompletion
 		{
 		private:
diff --git a/Engine/Source/Runtime/Interchange/Engine/Private/Tasks/InterchangeTaskParsing.cpp b/Engine/Source/Runtime/Interchange/Engine/Private/Tasks/InterchangeTaskParsing.cpp
--- a/Engine/Source/Runtime/Interchange/Engine/Private/Tasks/InterchangeTaskParsing.cpp
+++ b/Engine/Source/Runtime/Interchange/Engine/Private/Tasks/InterchangeTaskParsing.cpp
@@ -43,7 +43,7 @@ void UE::Interchange::FTaskParsing::DoTask(ENamedThreads::Type CurrentThread, co
 	};
 
 
-	FGraphEventArray CompletionPrerequistes;
+	
 	TArray<FTaskData> TaskDatas;
 
 	//Avoid creating asset if the asynchronous import is cancel, just create the completion task
@@ -62,6 +62,12 @@ void UE::Interchange::FTaskParsing::DoTask(ENamedThreads::Type CurrentThread, co
 			}
 			BaseNodeContainer->IterateNodes([&](const FString& NodeUID, UInterchangeBaseNode* Node)
 			{
+				if (!Node->IsEnabled())
+				{
+					//Do not call factory for a disable node
+					return;
+				}
+
 				if (Node->GetAssetClass() != nullptr)
 				{
 					const UClass* RegisteredFactoryClass = InterchangeManager->GetRegisteredFactoryClass(Node->GetAssetClass());
@@ -96,6 +102,8 @@ void UE::Interchange::FTaskParsing::DoTask(ENamedThreads::Type CurrentThread, co
 		return A.Dependencies.Num() <= B.Dependencies.Num();
 	});
 
+	FGraphEventArray CompletionPrerequistes;
+
 	for (int32 TaskIndex = 0; TaskIndex < TaskDatas.Num(); ++TaskIndex)
 	{
 		if (TaskDatas[TaskIndex].Dependencies.Num() > 0)
@@ -123,33 +131,39 @@ void UE::Interchange::FTaskParsing::DoTask(ENamedThreads::Type CurrentThread, co
 		);
 		CreatePackagePrerequistes.Add(AsyncHelper->CreatePackageTasks[CreatePackageTaskIndex]);
 
-		FGraphEventArray PostPipelinePrerequistes;
+		FGraphEventArray CreateAssetPrerequistes;
 		int32 CreateTaskIndex = AsyncHelper->CreateAssetTasks.Add(
 			TGraphTask<FTaskCreateAsset>::CreateTask(&(CreatePackagePrerequistes)).ConstructAndDispatchWhenReady(PackageBasePath, SourceIndex, WeakAsyncHelper, FactoryNode, bFactoryCanRunOnAnyThread)
 		);
-		PostPipelinePrerequistes.Add(AsyncHelper->CreateAssetTasks[CreateTaskIndex]);
+		CreateAssetPrerequistes.Add(AsyncHelper->CreateAssetTasks[CreateTaskIndex]);
 		
 		TaskDatas[TaskIndex].GraphEventRef = AsyncHelper->CreateAssetTasks[CreateTaskIndex];
+		CompletionPrerequistes.Add(TaskDatas[TaskIndex].GraphEventRef);
+	}
 
+	//Add a async task for pre completion
+	
+
+	FGraphEventArray PreCompletionPrerequistes;
+	AsyncHelper->PreCompletionTask = TGraphTask<FTaskPreCompletion>::CreateTask(&CompletionPrerequistes).ConstructAndDispatchWhenReady(InterchangeManager, WeakAsyncHelper);
+	PreCompletionPrerequistes.Add(AsyncHelper->PreCompletionTask);
+
+	//Start the Post pipeline task
+	for (int32 TaskIndex = 0; TaskIndex < TaskDatas.Num(); ++TaskIndex)
+	{
 		for (int32 GraphPipelineIndex = 0; GraphPipelineIndex < AsyncHelper->Pipelines.Num(); ++GraphPipelineIndex)
 		{
-			int32 GraphPipelineTaskIndex = INDEX_NONE;
-			GraphPipelineTaskIndex = AsyncHelper->PipelinePostImportTasks.Add(
-				TGraphTask<FTaskPipelinePostImport>::CreateTask(&(PostPipelinePrerequistes)).ConstructAndDispatchWhenReady(SourceIndex, GraphPipelineIndex, WeakAsyncHelper)
+			int32 GraphPipelineTaskIndex = AsyncHelper->PipelinePostImportTasks.Add(
+				TGraphTask<FTaskPipelinePostImport>::CreateTask(&(PreCompletionPrerequistes)).ConstructAndDispatchWhenReady(TaskDatas[TaskIndex].SourceIndex, GraphPipelineIndex, WeakAsyncHelper)
 			);
 			//Ensure we run the pipeline in the same order we create the task, since pipeline modify the node container, its important that its not process in parallel, Adding the one we start to the prerequisites
 			//is the way to go here
-			PostPipelinePrerequistes.Add(AsyncHelper->PipelinePostImportTasks[GraphPipelineTaskIndex]);
-
-			//Override the completion prerequisite with the latest post import pipeline task
-			TaskDatas[TaskIndex].GraphEventRef = AsyncHelper->PipelinePostImportTasks[GraphPipelineTaskIndex];
+			PreCompletionPrerequistes.Add(AsyncHelper->PipelinePostImportTasks[GraphPipelineTaskIndex]);
 		}
-
-		CompletionPrerequistes.Add(TaskDatas[TaskIndex].GraphEventRef);
 	}
 
 	FGraphEventArray PreAsyncCompletionPrerequistes;
-	AsyncHelper->PreAsyncCompletionTask = TGraphTask<FTaskPreAsyncCompletion>::CreateTask(&CompletionPrerequistes).ConstructAndDispatchWhenReady(InterchangeManager, WeakAsyncHelper);
+	AsyncHelper->PreAsyncCompletionTask = TGraphTask<FTaskPreAsyncCompletion>::CreateTask(&PreCompletionPrerequistes).ConstructAndDispatchWhenReady(InterchangeManager, WeakAsyncHelper);
 	PreAsyncCompletionPrerequistes.Add(AsyncHelper->PreAsyncCompletionTask);
 
 	AsyncHelper->CompletionTask = TGraphTask<FTaskCompletion>::CreateTask(&PreAsyncCompletionPrerequistes).ConstructAndDispatchWhenReady(InterchangeManager, WeakAsyncHelper);
diff --git a/Engine/Source/Runtime/Interchange/Engine/Private/Tasks/InterchangeTaskPipeline.cpp b/Engine/Source/Runtime/Interchange/Engine/Private/Tasks/InterchangeTaskPipeline.cpp
--- a/Engine/Source/Runtime/Interchange/Engine/Private/Tasks/InterchangeTaskPipeline.cpp
+++ b/Engine/Source/Runtime/Interchange/Engine/Private/Tasks/InterchangeTaskPipeline.cpp
@@ -57,7 +57,7 @@ void UE::Interchange::FTaskPipelinePostImport::DoTask(ENamedThreads::Type Curren
 		return;
 	}
 	UInterchangePipelineBase* PipelineBase = AsyncHelper->Pipelines[PipelineIndex];
-	TArray<FName> NodeUniqueIDs;
+	TArray<FString> NodeUniqueIDs;
 	TArray<UObject*> ImportAssets;
 	//Create a lock scope to read the Imported asset infos map
 	{
@@ -69,7 +69,7 @@ void UE::Interchange::FTaskPipelinePostImport::DoTask(ENamedThreads::Type Curren
 			ImportAssets.Reserve(ImportedInfos.Num());
 			for (UE::Interchange::FImportAsyncHelper::FImportedAssetInfo& ImportedInfo : ImportedInfos)
 			{
-				NodeUniqueIDs.Add(FName(*ImportedInfo.NodeUniqueId));
+				NodeUniqueIDs.Add(ImportedInfo.NodeUniqueId);
 				ImportAssets.Add(ImportedInfo.ImportAsset);
 			}
 		}
diff --git a/Engine/Source/Runtime/Interchange/Engine/Public/InterchangeManager.h b/Engine/Source/Runtime/Interchange/Engine/Public/InterchangeManager.h
--- a/Engine/Source/Runtime/Interchange/Engine/Public/InterchangeManager.h
+++ b/Engine/Source/Runtime/Interchange/Engine/Public/InterchangeManager.h
@@ -162,6 +162,7 @@ namespace UE
 			TArray<FGraphEventRef> CreateAssetTasks;
 
 			FGraphEventRef PreAsyncCompletionTask;
+			FGraphEventRef PreCompletionTask;
 			FGraphEventRef CompletionTask;
 
 			//Create package map, Key is package name. We cannot create package asynchronously so we have to create a game thread task to do this
