Fixed crash when importing a static mesh while opting not to import materials in the pipeline.

diff --git a/Engine/Plugins/Runtime/Interchange/Source/Import/Private/Mesh/InterchangeStaticMeshFactory.cpp b/Engine/Plugins/Runtime/Interchange/Source/Import/Private/Mesh/InterchangeStaticMeshFactory.cpp
--- a/Engine/Plugins/Runtime/Interchange/Source/Import/Private/Mesh/InterchangeStaticMeshFactory.cpp
+++ b/Engine/Plugins/Runtime/Interchange/Source/Import/Private/Mesh/InterchangeStaticMeshFactory.cpp
@@ -329,11 +329,7 @@ UObject* UInterchangeStaticMeshFactory::CreateAsset(const UInterchangeStaticMesh
 					int32 MaterialSlotIndex = StaticMesh->GetMaterialIndexFromImportedMaterialSlotName(MaterialSlotName);
 					if (MaterialSlotIndex == INDEX_NONE)
 					{
-#if WITH_EDITORONLY_DATA
-						StaticMesh->GetStaticMaterials().Emplace(MaterialInterface, MaterialSlotName, MaterialSlotName);
-#else
 						StaticMesh->GetStaticMaterials().Emplace(MaterialInterface, MaterialSlotName);
-#endif
 					}
 				}
 
@@ -343,7 +339,13 @@ UObject* UInterchangeStaticMeshFactory::CreateAsset(const UInterchangeStaticMesh
 				for (FPolygonGroupID PolygonGroupID : LodMeshDescription->PolygonGroups().GetElementIDs())
 				{
 					int32 MaterialSlotIndex = StaticMesh->GetMaterialIndexFromImportedMaterialSlotName(SlotNames[PolygonGroupID]);
-					check(MaterialSlotIndex != INDEX_NONE);
+					
+					if (MaterialSlotIndex == INDEX_NONE)
+					{
+						// If no material was found with this slot name, it is probably because the pipeline is configured to not import materials.
+						// Fill out a blank slot instead.
+						MaterialSlotIndex = StaticMesh->GetStaticMaterials().Emplace(nullptr, SlotNames[PolygonGroupID]);
+					}
 
 					FMeshSectionInfo Info = StaticMesh->GetSectionInfoMap().Get(CurrentLodIndex, SectionIndex);
 					Info.MaterialIndex = MaterialSlotIndex;
