Added missing header for non-unity builds.

diff --git a/Engine/Plugins/Runtime/Interchange/Source/Pipelines/Private/InterchangeGenericSkeletalMeshPipeline.cpp b/Engine/Plugins/Runtime/Interchange/Source/Pipelines/Private/InterchangeGenericSkeletalMeshPipeline.cpp
--- a/Engine/Plugins/Runtime/Interchange/Source/Pipelines/Private/InterchangeGenericSkeletalMeshPipeline.cpp
+++ b/Engine/Plugins/Runtime/Interchange/Source/Pipelines/Private/InterchangeGenericSkeletalMeshPipeline.cpp
@@ -3,6 +3,7 @@
 
 #include "Animation/Skeleton.h"
 #include "AssetRegistry/AssetRegistryModule.h"
+#include "Async/Async.h"
 #include "CoreMinimal.h"
 #include "InterchangeMaterialFactoryNode.h"
 #include "InterchangeMeshNode.h"
