Fix non unity build

diff --git a/Engine/Plugins/Runtime/Interchange/Source/Import/Private/Mesh/InterchangeSkeletalMeshFactory.cpp b/Engine/Plugins/Runtime/Interchange/Source/Import/Private/Mesh/InterchangeSkeletalMeshFactory.cpp
--- a/Engine/Plugins/Runtime/Interchange/Source/Import/Private/Mesh/InterchangeSkeletalMeshFactory.cpp
+++ b/Engine/Plugins/Runtime/Interchange/Source/Import/Private/Mesh/InterchangeSkeletalMeshFactory.cpp
@@ -20,6 +20,7 @@
 #include "Mesh/InterchangeSkeletalMeshPayloadInterface.h"
 #include "Nodes/InterchangeBaseNode.h"
 #include "Nodes/InterchangeBaseNodeContainer.h"
+#include "PhysicsEngine\PhysicsAsset.h"
 #include "Rendering/SkeletalMeshLODImporterData.h"
 #include "Rendering/SkeletalMeshModel.h"
 #include "SkeletalMeshAttributes.h"
diff --git a/Engine/Plugins/Runtime/Interchange/Source/Nodes/Public/InterchangePhysicsAssetFactoryNode.h b/Engine/Plugins/Runtime/Interchange/Source/Nodes/Public/InterchangePhysicsAssetFactoryNode.h
--- a/Engine/Plugins/Runtime/Interchange/Source/Nodes/Public/InterchangePhysicsAssetFactoryNode.h
+++ b/Engine/Plugins/Runtime/Interchange/Source/Nodes/Public/InterchangePhysicsAssetFactoryNode.h
@@ -70,7 +70,7 @@ public:
 	{
 		ensure(bIsNodeClassInitialized);
 #if WITH_ENGINE
-		return AssetClass.Get() != nullptr ? AssetClass.Get() : USkeleton::StaticClass();
+		return AssetClass.Get() != nullptr ? AssetClass.Get() : UPhysicsAsset::StaticClass();
 #else
 		return nullptr;
 #endif
