Fix static analysis

diff --git a/Engine/Plugins/Runtime/Interchange/Source/Import/Private/Mesh/InterchangeSkeletalMeshFactory.cpp b/Engine/Plugins/Runtime/Interchange/Source/Import/Private/Mesh/InterchangeSkeletalMeshFactory.cpp
--- a/Engine/Plugins/Runtime/Interchange/Source/Import/Private/Mesh/InterchangeSkeletalMeshFactory.cpp
+++ b/Engine/Plugins/Runtime/Interchange/Source/Import/Private/Mesh/InterchangeSkeletalMeshFactory.cpp
@@ -20,7 +20,7 @@
 #include "Mesh/InterchangeSkeletalMeshPayloadInterface.h"
 #include "Nodes/InterchangeBaseNode.h"
 #include "Nodes/InterchangeBaseNodeContainer.h"
-#include "PhysicsEngine\PhysicsAsset.h"
+#include "PhysicsEngine/PhysicsAsset.h"
 #include "Rendering/SkeletalMeshLODImporterData.h"
 #include "Rendering/SkeletalMeshModel.h"
 #include "SkeletalMeshAttributes.h"
