Making the reimport of a texture by interchange thread safe.

Changed the texture payloads so that they use the new unique buffer.
Refactored the blocked texture payload and the UDIM translator to support the new payload.
Added some functions to the FTextureSource so that it can receive a FSharedBuffer instead of copying the data into the payload.
Removed the const requirement from the UInterchangeFactoryBase main functions. It was needed so that I can store the texture payload in the function CreateAsset and then consume it in the function PreImportPreCompletedCallback. Also reporting messages in the factory does affect the state of the factory.
The Interface IImageWrapper has a new function to get the raw data into an ArrayView. For now the interface provides a default implementation but I would like to make it abstract when all the existing wrapper will have an implementation for it.

diff --git a/Engine/Plugins/Runtime/Interchange/Source/Import/Private/Material/InterchangeMaterialFactory.cpp b/Engine/Plugins/Runtime/Interchange/Source/Import/Private/Material/InterchangeMaterialFactory.cpp
--- a/Engine/Plugins/Runtime/Interchange/Source/Import/Private/Material/InterchangeMaterialFactory.cpp
+++ b/Engine/Plugins/Runtime/Interchange/Source/Import/Private/Material/InterchangeMaterialFactory.cpp
@@ -29,7 +29,7 @@ UClass* UInterchangeMaterialFactory::GetFactoryClass() const
 	return UMaterialInterface::StaticClass();
 }
 
-UObject* UInterchangeMaterialFactory::CreateEmptyAsset(const FCreateAssetParams& Arguments) const
+UObject* UInterchangeMaterialFactory::CreateEmptyAsset(const FCreateAssetParams& Arguments)
 {
 	UObject* Material = nullptr;
 
@@ -82,7 +82,7 @@ UObject* UInterchangeMaterialFactory::CreateEmptyAsset(const FCreateAssetParams&
 	return Material;
 }
 
-UObject* UInterchangeMaterialFactory::CreateAsset(const UInterchangeMaterialFactory::FCreateAssetParams& Arguments) const
+UObject* UInterchangeMaterialFactory::CreateAsset(const FCreateAssetParams& Arguments)
 {
 #if !WITH_EDITORONLY_DATA
 
@@ -302,7 +302,7 @@ UObject* UInterchangeMaterialFactory::CreateAsset(const UInterchangeMaterialFact
 }
 
 /* This function is call in the completion task on the main thread, use it to call main thread post creation step for your assets*/
-void UInterchangeMaterialFactory::PreImportPreCompletedCallback(const FImportPreCompletedCallbackParams& Arguments) const
+void UInterchangeMaterialFactory::PreImportPreCompletedCallback(const FImportPreCompletedCallbackParams& Arguments)
 {
 	check(IsInGameThread());
 	Super::PreImportPreCompletedCallback(Arguments);
diff --git a/Engine/Plugins/Runtime/Interchange/Source/Import/Private/Mesh/InterchangePhysicsAssetFactory.cpp b/Engine/Plugins/Runtime/Interchange/Source/Import/Private/Mesh/InterchangePhysicsAssetFactory.cpp
--- a/Engine/Plugins/Runtime/Interchange/Source/Import/Private/Mesh/InterchangePhysicsAssetFactory.cpp
+++ b/Engine/Plugins/Runtime/Interchange/Source/Import/Private/Mesh/InterchangePhysicsAssetFactory.cpp
@@ -14,7 +14,7 @@ UClass* UInterchangePhysicsAssetFactory::GetFactoryClass() const
 	return UPhysicsAsset::StaticClass();
 }
 
-UObject* UInterchangePhysicsAssetFactory::CreateEmptyAsset(const FCreateAssetParams& Arguments) const
+UObject* UInterchangePhysicsAssetFactory::CreateEmptyAsset(const FCreateAssetParams& Arguments)
 {
 	UObject* PhysicsAsset = nullptr;
 
@@ -56,7 +56,7 @@ UObject* UInterchangePhysicsAssetFactory::CreateEmptyAsset(const FCreateAssetPar
 	return PhysicsAsset;
 }
 
-UObject* UInterchangePhysicsAssetFactory::CreateAsset(const UInterchangePhysicsAssetFactory::FCreateAssetParams& Arguments) const
+UObject* UInterchangePhysicsAssetFactory::CreateAsset(const FCreateAssetParams& Arguments)
 {
 #if !WITH_EDITORONLY_DATA
 
diff --git a/Engine/Plugins/Runtime/Interchange/Source/Import/Private/Mesh/InterchangeSkeletalMeshFactory.cpp b/Engine/Plugins/Runtime/Interchange/Source/Import/Private/Mesh/InterchangeSkeletalMeshFactory.cpp
--- a/Engine/Plugins/Runtime/Interchange/Source/Import/Private/Mesh/InterchangeSkeletalMeshFactory.cpp
+++ b/Engine/Plugins/Runtime/Interchange/Source/Import/Private/Mesh/InterchangeSkeletalMeshFactory.cpp
@@ -669,7 +669,7 @@ UClass* UInterchangeSkeletalMeshFactory::GetFactoryClass() const
 	return USkeletalMesh::StaticClass();
 }
 
-UObject* UInterchangeSkeletalMeshFactory::CreateEmptyAsset(const FCreateAssetParams& Arguments) const
+UObject* UInterchangeSkeletalMeshFactory::CreateEmptyAsset(const FCreateAssetParams& Arguments)
 {
 #if !WITH_EDITOR || !WITH_EDITORONLY_DATA
 
@@ -717,7 +717,7 @@ UObject* UInterchangeSkeletalMeshFactory::CreateEmptyAsset(const FCreateAssetPar
 #endif //else !WITH_EDITOR || !WITH_EDITORONLY_DATA
 }
 
-UObject* UInterchangeSkeletalMeshFactory::CreateAsset(const UInterchangeSkeletalMeshFactory::FCreateAssetParams& Arguments) const
+UObject* UInterchangeSkeletalMeshFactory::CreateAsset(const FCreateAssetParams& Arguments)
 {
 #if !WITH_EDITOR || !WITH_EDITORONLY_DATA
 
@@ -1054,7 +1054,7 @@ UObject* UInterchangeSkeletalMeshFactory::CreateAsset(const UInterchangeSkeletal
 }
 
 /* This function is call in the completion task on the main thread, use it to call main thread post creation step for your assets*/
-void UInterchangeSkeletalMeshFactory::PreImportPreCompletedCallback(const FImportPreCompletedCallbackParams& Arguments) const
+void UInterchangeSkeletalMeshFactory::PreImportPreCompletedCallback(const FImportPreCompletedCallbackParams& Arguments)
 {
 	check(IsInGameThread());
 	Super::PreImportPreCompletedCallback(Arguments);
diff --git a/Engine/Plugins/Runtime/Interchange/Source/Import/Private/Mesh/InterchangeSkeletonFactory.cpp b/Engine/Plugins/Runtime/Interchange/Source/Import/Private/Mesh/InterchangeSkeletonFactory.cpp
--- a/Engine/Plugins/Runtime/Interchange/Source/Import/Private/Mesh/InterchangeSkeletonFactory.cpp
+++ b/Engine/Plugins/Runtime/Interchange/Source/Import/Private/Mesh/InterchangeSkeletonFactory.cpp
@@ -15,7 +15,7 @@ UClass* UInterchangeSkeletonFactory::GetFactoryClass() const
 	return USkeleton::StaticClass();
 }
 
-UObject* UInterchangeSkeletonFactory::CreateEmptyAsset(const FCreateAssetParams& Arguments) const
+UObject* UInterchangeSkeletonFactory::CreateEmptyAsset(const FCreateAssetParams& Arguments)
 {
 	UObject* Skeleton = nullptr;
 
@@ -57,7 +57,7 @@ UObject* UInterchangeSkeletonFactory::CreateEmptyAsset(const FCreateAssetParams&
 	return Skeleton;
 }
 
-UObject* UInterchangeSkeletonFactory::CreateAsset(const UInterchangeSkeletonFactory::FCreateAssetParams& Arguments) const
+UObject* UInterchangeSkeletonFactory::CreateAsset(const FCreateAssetParams& Arguments)
 {
 #if !WITH_EDITORONLY_DATA
 
diff --git a/Engine/Plugins/Runtime/Interchange/Source/Import/Private/Mesh/InterchangeStaticMeshFactory.cpp b/Engine/Plugins/Runtime/Interchange/Source/Import/Private/Mesh/InterchangeStaticMeshFactory.cpp
--- a/Engine/Plugins/Runtime/Interchange/Source/Import/Private/Mesh/InterchangeStaticMeshFactory.cpp
+++ b/Engine/Plugins/Runtime/Interchange/Source/Import/Private/Mesh/InterchangeStaticMeshFactory.cpp
@@ -33,7 +33,7 @@ UClass* UInterchangeStaticMeshFactory::GetFactoryClass() const
 }
 
 
-UObject* UInterchangeStaticMeshFactory::CreateEmptyAsset(const FCreateAssetParams& Arguments) const
+UObject* UInterchangeStaticMeshFactory::CreateEmptyAsset(const FCreateAssetParams& Arguments)
 {
 #if !WITH_EDITOR || !WITH_EDITORONLY_DATA
 
@@ -79,7 +79,7 @@ UObject* UInterchangeStaticMeshFactory::CreateEmptyAsset(const FCreateAssetParam
 }
 
 
-UObject* UInterchangeStaticMeshFactory::CreateAsset(const UInterchangeStaticMeshFactory::FCreateAssetParams& Arguments) const
+UObject* UInterchangeStaticMeshFactory::CreateAsset(const FCreateAssetParams& Arguments)
 {
 #if !WITH_EDITOR || !WITH_EDITORONLY_DATA
 
@@ -387,7 +387,7 @@ UObject* UInterchangeStaticMeshFactory::CreateAsset(const UInterchangeStaticMesh
 
 
 /* This function is call in the completion task on the main thread, use it to call main thread post creation step for your assets */
-void UInterchangeStaticMeshFactory::PreImportPreCompletedCallback(const FImportPreCompletedCallbackParams& Arguments) const
+void UInterchangeStaticMeshFactory::PreImportPreCompletedCallback(const FImportPreCompletedCallbackParams& Arguments)
 {
 	check(IsInGameThread());
 	Super::PreImportPreCompletedCallback(Arguments);
diff --git a/Engine/Plugins/Runtime/Interchange/Source/Import/Private/Texture/InterchangeBMPTranslator.cpp b/Engine/Plugins/Runtime/Interchange/Source/Import/Private/Texture/InterchangeBMPTranslator.cpp
--- a/Engine/Plugins/Runtime/Interchange/Source/Import/Private/Texture/InterchangeBMPTranslator.cpp
+++ b/Engine/Plugins/Runtime/Interchange/Source/Import/Private/Texture/InterchangeBMPTranslator.cpp
@@ -99,7 +99,7 @@ TOptional<UE::Interchange::FImportImage> UInterchangeBMPTranslator::GetTexturePa
 		false
 	);
 
-	if (!BmpImageWrapper->GetRaw(Format, BitDepth, PayloadData.RawData))
+	if (!BmpImageWrapper->GetRaw(Format, BitDepth, PayloadData.GetArrayViewOfRawData()))
 	{
 		UE_LOG(LogInterchangeImport, Error, TEXT("Failed to decode BMP. [%s]"), *Filename);
 		return TOptional<UE::Interchange::FImportImage>();
diff --git a/Engine/Plugins/Runtime/Interchange/Source/Import/Private/Texture/InterchangeBlockedTexturePayloadData.cpp b/Engine/Plugins/Runtime/Interchange/Source/Import/Private/Texture/InterchangeBlockedTexturePayloadData.cpp
new file mode 100644
--- /dev/null
+++ b/Engine/Plugins/Runtime/Interchange/Source/Import/Private/Texture/InterchangeBlockedTexturePayloadData.cpp
@@ -0,0 +1,135 @@
+// Copyright Epic Games, Inc. All Rights Reserved.
+
+#include "Texture/InterchangeBlockedTexturePayloadData.h"
+
+#include "Async/ParallelFor.h"
+
+namespace UE::Interchange
+{
+	namespace Private::BlockedPayload
+	{
+		int64 ComputeBlockSize(const FTextureSourceBlock& Block, ETextureSourceFormat Format)
+		{
+			int64 Size = 0;
+			int32 Width = Block.SizeX;
+			int32 Height = Block.SizeY;
+			for (int32 CurrentMip = 0; CurrentMip < Block.NumMips; ++CurrentMip)
+			{
+				Size += Width * Height * FTextureSource::GetBytesPerPixel(Format) * Block.NumSlices;
+				Width >>= 2;
+				Height >>= 2;
+			}
+			return Size;
+		}
+	}
+
+	int64 FImportBlockedImage::ComputeBufferSize() const
+	{
+		int64 Size = 0;
+		for (const FTextureSourceBlock& Block : BlocksData)
+		{
+			Size += Private::BlockedPayload::ComputeBlockSize(Block, Format);
+		}
+
+		return Size;
+	}
+
+	bool FImportBlockedImage::IsValid() const
+	{
+		return Format != TSF_Invalid
+			&& BlocksData.Num() > 0
+			&& RawData.GetSize() == ComputeBufferSize()
+			;
+	}
+
+	bool FImportBlockedImage::InitDataSharedAmongBlocks(const FImportImage& Image)
+	{
+		if (Image.IsValid())
+		{
+			Format = Image.Format;
+			CompressionSettings = Image.CompressionSettings;
+			bSRGB = Image.bSRGB;
+			MipGenSettings = Image.MipGenSettings;
+			return true;
+		}
+
+		return false;
+	}
+
+	bool FImportBlockedImage::InitBlockFromImage(int32 BlockX, int32 BlockY, const FImportImage& Image)
+	{
+		if (!Image.IsValid())
+		{
+			return false;
+		}
+
+		if (Image.Format != Format)
+		{
+			// The images must be of the same format
+			return false;
+		}
+
+		FTextureSourceBlock& Block = BlocksData.AddDefaulted_GetRef();
+		Block.BlockX = BlockX;
+		Block.BlockY = BlockY;
+		Block.SizeX = Image.SizeX;
+		Block.SizeY = Image.SizeY;
+		Block.NumSlices = 1;
+		Block.NumMips = Image.NumMips;
+
+		return true;
+	}
+
+	bool FImportBlockedImage::MigrateDataFromImagesToRawData(TArray<FImportImage>& Images)
+	{
+		RawData = FUniqueBuffer::Alloc(ComputeBufferSize());
+
+		int64 ImageSize = 0;
+		TArray<int64> BlocksOffset;
+		BlocksOffset.Reserve(Images.Num());
+		
+		if (BlocksData.Num() != Images.Num())
+		{
+			return false;
+		}
+
+		int32 BlockIndex = 0;
+		int64 CurrentBlockOffset = 0;
+		for (const FImportImage& Image : Images)
+		{
+			if (!Image.IsValid())
+			{
+				return false;
+			}
+
+			if (Image.Format != Format)
+			{
+				// The images must be of the same format
+				return false;
+			}
+
+			BlocksOffset.Add(CurrentBlockOffset);
+			const int64 BlockSize = Private::BlockedPayload::ComputeBlockSize(BlocksData[BlockIndex], Format);
+			CurrentBlockOffset += BlockSize;
+
+			if (BlockSize != Image.RawData.GetSize())
+			{
+				return false;
+			}
+
+			++BlockIndex;
+		}
+
+		uint8* StartPtr = static_cast<uint8*>(RawData.GetData());
+		ParallelFor(Images.Num(),[this, StartPtr, &Images, &BlocksOffset](int32 Index)
+			{
+				int64 BlockSize = Private::BlockedPayload::ComputeBlockSize(BlocksData[Index], Format);
+				FMemory::Memcpy(StartPtr + BlocksOffset[Index], static_cast<uint8*>(Images[Index].RawData.GetData()), BlockSize);
+			}
+			, IsInGameThread() ? EParallelForFlags::Unbalanced : EParallelForFlags::BackgroundPriority | EParallelForFlags::Unbalanced
+			);
+
+
+		return true;
+	}
+}
\ No newline at end of file
diff --git a/Engine/Plugins/Runtime/Interchange/Source/Import/Private/Texture/InterchangeDDSTranslator.cpp b/Engine/Plugins/Runtime/Interchange/Source/Import/Private/Texture/InterchangeDDSTranslator.cpp
--- a/Engine/Plugins/Runtime/Interchange/Source/Import/Private/Texture/InterchangeDDSTranslator.cpp
+++ b/Engine/Plugins/Runtime/Interchange/Source/Import/Private/Texture/InterchangeDDSTranslator.cpp
@@ -288,18 +288,13 @@ TOptional<UE::Interchange::FImportSlicedImage> UInterchangeDDSTranslator::GetSli
 
 	uint8* DestMipData[MAX_TEXTURE_MIP_COUNT] = { 0 };
 	int64 MipsSize[MAX_TEXTURE_MIP_COUNT] = { 0 };
-	int64 NumBytes = 0;
 
 	for (uint32 MipIndex = 0; MipIndex < MipMapCount; ++MipIndex)
 	{
 		int64 MipSize = PayloadData.GetMipSize(MipIndex);
 		MipsSize[MipIndex] = MipSize;
-		// Six time since we have six slice
-		NumBytes += MipSize * 6;
 	}
 
-	PayloadData.RawData.AddUninitialized(NumBytes);
-
 	for (uint32 MipIndex = 0; MipIndex < MipMapCount; ++MipIndex)
 	{
 		DestMipData[MipIndex] = PayloadData.GetMipData(MipIndex);
diff --git a/Engine/Plugins/Runtime/Interchange/Source/Import/Private/Texture/InterchangeEXRTranslator.cpp b/Engine/Plugins/Runtime/Interchange/Source/Import/Private/Texture/InterchangeEXRTranslator.cpp
--- a/Engine/Plugins/Runtime/Interchange/Source/Import/Private/Texture/InterchangeEXRTranslator.cpp
+++ b/Engine/Plugins/Runtime/Interchange/Source/Import/Private/Texture/InterchangeEXRTranslator.cpp
@@ -105,7 +105,7 @@ TOptional<UE::Interchange::FImportImage> UInterchangeEXRTranslator::GetTexturePa
 	);
 	PayloadData.CompressionSettings = TC_HDR;
 
-	if (!ExrImageWrapper->GetRaw(Format, BitDepth, PayloadData.RawData))
+	if (!ExrImageWrapper->GetRaw(Format, BitDepth, PayloadData.GetArrayViewOfRawData()))
 	{
 		UE_LOG(LogInterchangeImport, Error, TEXT("Failed to decode EXR. [%s]"), *Filename);
 		return TOptional<UE::Interchange::FImportImage>();
diff --git a/Engine/Plugins/Runtime/Interchange/Source/Import/Private/Texture/InterchangeHDRTranslator.cpp b/Engine/Plugins/Runtime/Interchange/Source/Import/Private/Texture/InterchangeHDRTranslator.cpp
--- a/Engine/Plugins/Runtime/Interchange/Source/Import/Private/Texture/InterchangeHDRTranslator.cpp
+++ b/Engine/Plugins/Runtime/Interchange/Source/Import/Private/Texture/InterchangeHDRTranslator.cpp
@@ -90,7 +90,7 @@ TOptional<UE::Interchange::FImportImage> UInterchangeHDRTranslator::GetTexturePa
 			TSF_BGRE8,
 			false);
 
-		if (!HdrImageWrapper.GetRaw(ERGBFormat::BGRE, 8, PayloadData.RawData))
+		if (!HdrImageWrapper.GetRaw(ERGBFormat::BGRE, 8, PayloadData.GetArrayViewOfRawData()))
 		{
 			UE_LOG(LogInterchangeImport, Error, TEXT("Failed to Import HDR. [%s]"), *Filename);
 			UE_LOG(LogInterchangeImport, Error, TEXT("%s"), *( HdrImageWrapper.GetErrorMessage().ToString()));
diff --git a/Engine/Plugins/Runtime/Interchange/Source/Import/Private/Texture/InterchangeIESTranslator.cpp b/Engine/Plugins/Runtime/Interchange/Source/Import/Private/Texture/InterchangeIESTranslator.cpp
--- a/Engine/Plugins/Runtime/Interchange/Source/Import/Private/Texture/InterchangeIESTranslator.cpp
+++ b/Engine/Plugins/Runtime/Interchange/Source/Import/Private/Texture/InterchangeIESTranslator.cpp
@@ -74,7 +74,10 @@ TOptional<UE::Interchange::FImportLightProfile> UInterchangeIESTranslator::GetLi
 		Payload.Brightness = IESConverter.GetBrightness();
 		Payload.TextureMultiplier = IESConverter.GetMultiplier();
 
-		Payload.RawData = IESConverter.GetRawData();
+
+		const TArray<uint8>&  RawData = IESConverter.GetRawData();;
+
+		FPlatformMemory::Memcpy(Payload.RawData.GetData(), RawData.GetData(), Payload.RawData.GetSize());
 
 		return Payload;
 	}
diff --git a/Engine/Plugins/Runtime/Interchange/Source/Import/Private/Texture/InterchangeJPGTranslator.cpp b/Engine/Plugins/Runtime/Interchange/Source/Import/Private/Texture/InterchangeJPGTranslator.cpp
--- a/Engine/Plugins/Runtime/Interchange/Source/Import/Private/Texture/InterchangeJPGTranslator.cpp
+++ b/Engine/Plugins/Runtime/Interchange/Source/Import/Private/Texture/InterchangeJPGTranslator.cpp
@@ -124,7 +124,7 @@ TOptional<UE::Interchange::FImportImage> UInterchangeJPGTranslator::GetTexturePa
 		BitDepth < 16
 	);
 
-	if (!JpegImageWrapper->GetRaw(Format, BitDepth, PayloadData.RawData))
+	if (!JpegImageWrapper->GetRaw(Format, BitDepth, PayloadData.GetArrayViewOfRawData()))
 	{
 		UE_LOG(LogInterchangeImport, Error, TEXT("Failed to decode JPEG. [%s]"), *Filename);
 		return TOptional<UE::Interchange::FImportImage>();
diff --git a/Engine/Plugins/Runtime/Interchange/Source/Import/Private/Texture/InterchangePCXTranslator.cpp b/Engine/Plugins/Runtime/Interchange/Source/Import/Private/Texture/InterchangePCXTranslator.cpp
--- a/Engine/Plugins/Runtime/Interchange/Source/Import/Private/Texture/InterchangePCXTranslator.cpp
+++ b/Engine/Plugins/Runtime/Interchange/Source/Import/Private/Texture/InterchangePCXTranslator.cpp
@@ -185,7 +185,7 @@ TOptional<UE::Interchange::FImportImage> UInterchangePCXTranslator::GetTexturePa
 				TSF_BGRA8
 			);
 
-			uint8* Dest = PayloadData.RawData.GetData();
+			uint8* Dest = static_cast<uint8*>(PayloadData.RawData.GetData());
 
 			// Doing a memset to make sure the alpha channel is set to 0xff since we only have 3 color planes.
 			FMemory::Memset(Dest, 0xff, NewU * NewV * FTextureSource::GetBytesPerPixel(PayloadData.Format));
diff --git a/Engine/Plugins/Runtime/Interchange/Source/Import/Private/Texture/InterchangePNGTranslator.cpp b/Engine/Plugins/Runtime/Interchange/Source/Import/Private/Texture/InterchangePNGTranslator.cpp
--- a/Engine/Plugins/Runtime/Interchange/Source/Import/Private/Texture/InterchangePNGTranslator.cpp
+++ b/Engine/Plugins/Runtime/Interchange/Source/Import/Private/Texture/InterchangePNGTranslator.cpp
@@ -308,7 +308,7 @@ TOptional<UE::Interchange::FImportImage> UInterchangePNGTranslator::GetTexturePa
 		BitDepth < 16
 	);
 
-	if (PngImageWrapper->GetRaw(Format, BitDepth, PayloadData.RawData))
+	if (PngImageWrapper->GetRaw(Format, BitDepth, PayloadData.GetArrayViewOfRawData()))
 	{
 		bool bFillPNGZeroAlpha = true;
 		GConfig->GetBool(TEXT("TextureImporter"), TEXT("FillPNGZeroAlpha"), bFillPNGZeroAlpha, GEditorIni);
@@ -316,7 +316,7 @@ TOptional<UE::Interchange::FImportImage> UInterchangePNGTranslator::GetTexturePa
 		if (bFillPNGZeroAlpha)
 		{
 			// Replace the pixels with 0.0 alpha with a color value from the nearest neighboring color which has a non-zero alpha
-			UE::Interchange::Private::FillZeroAlphaPNGData(PayloadData.SizeX, PayloadData.SizeY, PayloadData.Format, PayloadData.RawData.GetData());
+			UE::Interchange::Private::FillZeroAlphaPNGData(PayloadData.SizeX, PayloadData.SizeY, PayloadData.Format, static_cast<uint8*>(PayloadData.RawData.GetData()));
 		}
 	}
 	else
diff --git a/Engine/Plugins/Runtime/Interchange/Source/Import/Private/Texture/InterchangePSDTranslator.cpp b/Engine/Plugins/Runtime/Interchange/Source/Import/Private/Texture/InterchangePSDTranslator.cpp
--- a/Engine/Plugins/Runtime/Interchange/Source/Import/Private/Texture/InterchangePSDTranslator.cpp
+++ b/Engine/Plugins/Runtime/Interchange/Source/Import/Private/Texture/InterchangePSDTranslator.cpp
@@ -475,7 +475,7 @@ TOptional<UE::Interchange::FImportImage> UInterchangePSDTranslator::GetTexturePa
 		TextureFormat
 	);
 
-	if (!UE::Interchange::Private::ReadData(PayloadData.RawData.GetData(), Buffer, PSDHeader))
+	if (!UE::Interchange::Private::ReadData(static_cast<uint8*>(PayloadData.RawData.GetData()), Buffer, PSDHeader))
 	{
 		UE_LOG(LogInterchangeImport, Error, TEXT("Failed to read this PSD"));
 		return TOptional<UE::Interchange::FImportImage>();
diff --git a/Engine/Plugins/Runtime/Interchange/Source/Import/Private/Texture/InterchangeSlicedTexturePayloadData.cpp b/Engine/Plugins/Runtime/Interchange/Source/Import/Private/Texture/InterchangeSlicedTexturePayloadData.cpp
--- a/Engine/Plugins/Runtime/Interchange/Source/Import/Private/Texture/InterchangeSlicedTexturePayloadData.cpp
+++ b/Engine/Plugins/Runtime/Interchange/Source/Import/Private/Texture/InterchangeSlicedTexturePayloadData.cpp
@@ -3,27 +3,13 @@
 #include "Texture/InterchangeSlicedTexturePayloadData.h"
 
 #include "CoreMinimal.h"
-#include "Engine/Texture.h"
 
 
 void UE::Interchange::FImportSlicedImage::Init(int32 InSizeX, int32 InSizeY, int32 InNumSlice, int32 InNumMips, ETextureSourceFormat InFormat, bool InSRGB)
 {
-	SizeX = InSizeX;
-	SizeY = InSizeY;
 	NumSlice = InNumSlice;
-	NumMips = InNumMips;
-	Format = InFormat;
-	bSRGB = InSRGB;
-}
-
-int64 UE::Interchange::FImportSlicedImage::GetMipSize(int32 InMipIndex) const
-{
-	check(InMipIndex >= 0);
-	check(InMipIndex < NumMips);
-	const int32 MipSizeX = FMath::Max(SizeX >> InMipIndex, 1);
-	const int32 MipSizeY = FMath::Max(SizeY >> InMipIndex, 1);
-	return (int64)MipSizeX * MipSizeY * FTextureSource::GetBytesPerPixel(Format);
 
+	FImportImage::Init2DWithParams(InSizeX, InSizeY, InNumMips, InFormat, InSRGB);
 }
 
 const uint8* UE::Interchange::FImportSlicedImage::GetMipData(int32 InMipIndex, int32 InSliceIndex) const
@@ -40,7 +26,9 @@ const uint8* UE::Interchange::FImportSlicedImage::GetMipData(int32 InMipIndex, i
 		Offset += GetMipSize(MipIndex) * InSliceIndex;
 	}
 
-	return &RawData[Offset];
+	check(RawData.GetSize() > uint64(Offset));
+
+	return static_cast<const uint8*>(RawData.GetData()) + Offset;
 }
 
 uint8* UE::Interchange::FImportSlicedImage::GetMipData(int32 InMipIndex, int32 InSliceIndex)
@@ -48,3 +36,14 @@ uint8* UE::Interchange::FImportSlicedImage::GetMipData(int32 InMipIndex, int32 I
 	const uint8* ConstBufferPtr = static_cast<const FImportSlicedImage*>(this)->GetMipData(InMipIndex, InSliceIndex);
 	return const_cast<uint8*>(ConstBufferPtr);
 }
+
+int64 UE::Interchange::FImportSlicedImage::ComputeBufferSize() const
+{
+	return FImportImage::ComputeBufferSize() * NumSlice;
+}
+
+bool UE::Interchange::FImportSlicedImage::IsValid() const
+{
+	return NumSlice > 0 && FImportImage::IsValid();
+}
+
diff --git a/Engine/Plugins/Runtime/Interchange/Source/Import/Private/Texture/InterchangeTGATranslator.cpp b/Engine/Plugins/Runtime/Interchange/Source/Import/Private/Texture/InterchangeTGATranslator.cpp
--- a/Engine/Plugins/Runtime/Interchange/Source/Import/Private/Texture/InterchangeTGATranslator.cpp
+++ b/Engine/Plugins/Runtime/Interchange/Source/Import/Private/Texture/InterchangeTGATranslator.cpp
@@ -141,12 +141,12 @@ TOptional<UE::Interchange::FImportImage> UInterchangeTGATranslator::GetTexturePa
 
 		TGATranslatorImpl::DecompressTGA( TGA, PayloadData );
 
-		const bool bResult = TgaImageWrapper->GetRaw( TgaImageWrapper->GetFormat(), TgaImageWrapper->GetBitDepth(), PayloadData.RawData );
+		const bool bResult = TgaImageWrapper->GetRaw( TgaImageWrapper->GetFormat(), TgaImageWrapper->GetBitDepth(), PayloadData.GetArrayViewOfRawData() );
 
 		if ( bResult && PayloadData.CompressionSettings == TC_Grayscale && TGA->ImageTypeCode == 3 )
 		{
 			// default grayscales to linear as they wont get compression otherwise and are commonly used as masks
-			PayloadData.SRGB = false;
+			PayloadData.bSRGB = false;
 		}
 
 		return PayloadData;
diff --git a/Engine/Plugins/Runtime/Interchange/Source/Import/Private/Texture/InterchangeTextureFactory.cpp b/Engine/Plugins/Runtime/Interchange/Source/Import/Private/Texture/InterchangeTextureFactory.cpp
--- a/Engine/Plugins/Runtime/Interchange/Source/Import/Private/Texture/InterchangeTextureFactory.cpp
+++ b/Engine/Plugins/Runtime/Interchange/Source/Import/Private/Texture/InterchangeTextureFactory.cpp
@@ -20,11 +20,10 @@
 #include "InterchangeTextureLightProfileNode.h"
 #include "InterchangeTextureNode.h"
 #include "InterchangeTranslatorBase.h"
-#include "Misc/TVariant.h"
 #include "Nodes/InterchangeBaseNode.h"
+#include "ProfilingDebugging/CpuProfilerTrace.h"
 #include "Texture/InterchangeBlockedTexturePayloadInterface.h"
 #include "Texture/InterchangeSlicedTexturePayloadInterface.h"
-#include "Texture/InterchangeTextureLightProfilePayloadData.h"
 #include "Texture/InterchangeTextureLightProfilePayloadInterface.h"
 #include "Texture/InterchangeTexturePayloadInterface.h"
 #include "TextureCompiler.h"
@@ -227,12 +226,6 @@ namespace UE::Interchange::Private::InterchangeTextureFactory
 		return {};
 	}
 
-	using FTexturePayloadVariant = TVariant<FEmptyVariantState
-		, TOptional<FImportImage>
-		, TOptional<FImportBlockedImage>
-		, TOptional<FImportSlicedImage>
-		, TOptional<FImportLightProfile>>;
-
 	FTexturePayloadVariant GetTexturePayload(const UInterchangeSourceData* SourceData, const FString& PayloadKey, const FTextureNodeVariant& TextureNodeVariant, const UInterchangeTranslatorBase* Translator)
 	{
 		// Standard texture 2D payload
@@ -270,7 +263,7 @@ namespace UE::Interchange::Private::InterchangeTextureFactory
 	}
 
 #if WITH_EDITORONLY_DATA
-	void SetupTextureSourceData(UTexture* Texture, const FImportImage& Image)
+	void SetupTextureSourceData(UTexture* Texture, FImportImage& Image)
 	{
 		Texture->Source.Init(
 			Image.SizeX,
@@ -278,11 +271,11 @@ namespace UE::Interchange::Private::InterchangeTextureFactory
 			/*NumSlices=*/ 1,
 			Image.NumMips,
 			Image.Format,
-			Image.RawData.GetData()
+			Image.RawData.MoveToShared()
 		);
 
 		Texture->CompressionSettings = Image.CompressionSettings;
-		Texture->SRGB = Image.SRGB;
+		Texture->SRGB = Image.bSRGB;
 
 		//If the MipGenSettings was set by the translator, we must apply it before the build
 		if (Image.MipGenSettings.IsSet())
@@ -292,50 +285,50 @@ namespace UE::Interchange::Private::InterchangeTextureFactory
 		}
 	}
 
-	bool SetupTexture2DSourceData(UTexture2D* Texture2D, const FImportBlockedImage& BlockedImage)
+	bool SetupTexture2DSourceData(UTexture2D* Texture2D, FImportBlockedImage& BlockedImage)
 	{
-		if (BlockedImage.IsBlockedData())
+		if (BlockedImage.IsValid())
 		{
-			if (BlockedImage.HasData())
+			if (BlockedImage.BlocksData.Num() > 1)
 			{
-				ETextureSourceFormat Format = BlockedImage.ImagesData[0].Format;
-				TextureCompressionSettings CompressionSettings = BlockedImage.ImagesData[0].CompressionSettings;
-				bool bSRGB = BlockedImage.ImagesData[0].SRGB;
-
-				TArray<const uint8*> SourceImageDatasPtr;
-				SourceImageDatasPtr.Reserve(BlockedImage.ImagesData.Num());
-				for (const  UE::Interchange::FImportImage& Image : BlockedImage.ImagesData)
-				{
-					SourceImageDatasPtr.Add(Image.RawData.GetData());
-				}
 
 				Texture2D->Source.InitBlocked(
-					&Format,
+					&BlockedImage.Format,
 					BlockedImage.BlocksData.GetData(),
-					/*NumSlices=*/ 1,
+					/*InNumLayers=*/ 1,
 					BlockedImage.BlocksData.Num(),
-					SourceImageDatasPtr.GetData()
-				);
+					BlockedImage.RawData.MoveToShared()
+					);
 
-				Texture2D->CompressionSettings = CompressionSettings;
-				Texture2D->SRGB = bSRGB;
+				Texture2D->CompressionSettings = BlockedImage.CompressionSettings;
+				Texture2D->SRGB = BlockedImage.bSRGB;
 				Texture2D->VirtualTextureStreaming = true;
 
-				if (BlockedImage.ImagesData[0].MipGenSettings.IsSet())
+				if (BlockedImage.MipGenSettings.IsSet())
 				{
 					// if the source has mips we keep the mips by default, unless the user changes that
-					Texture2D->MipGenSettings = BlockedImage.ImagesData[0].MipGenSettings.GetValue();
+					Texture2D->MipGenSettings = BlockedImage.MipGenSettings.GetValue();
 				}
 
 				return true;
 			}
-		}
-		else
-		{
-			//Import as a normal texture
-			if (BlockedImage.ImagesData.Num() == 1)
+			else
 			{
-				SetupTextureSourceData(Texture2D, BlockedImage.ImagesData[0]);
+				//Import as a normal texture
+				FImportImage Image;
+				Image.Format = BlockedImage.Format;
+				Image.CompressionSettings = BlockedImage.CompressionSettings;
+				Image.bSRGB = BlockedImage.bSRGB;
+				Image.MipGenSettings = BlockedImage.MipGenSettings;
+
+				const FTextureSourceBlock& Block = BlockedImage.BlocksData[0];
+				Image.SizeX = Block.SizeX;
+				Image.SizeY = Block.SizeY;
+				Image.NumMips = Block.NumMips;
+
+				Image.RawData = MoveTemp(BlockedImage.RawData);
+
+				SetupTextureSourceData(Texture2D, Image);
 				return true;
 			}
 		}
@@ -343,43 +336,21 @@ namespace UE::Interchange::Private::InterchangeTextureFactory
 		return false;
 	}
 
-	void SetupTextureSourceData(UTexture* Texture, const FImportSlicedImage& SlicedImage)
+	void SetupTextureSourceData(UTexture* Texture, FImportSlicedImage& SlicedImage)
 	{
-		Texture->Source.Init(
+		Texture->Source.InitLayered(
 			SlicedImage.SizeX,
 			SlicedImage.SizeY,
 			SlicedImage.NumSlice,
+			1,
 			SlicedImage.NumMips,
-			SlicedImage.Format
-		);
+			&SlicedImage.Format,
+			SlicedImage.RawData.MoveToShared()
+			);
 
 		Texture->CompressionSettings = SlicedImage.CompressionSettings;
 		Texture->SRGB = SlicedImage.bSRGB;
 
-
-		uint8* DestMipData[MAX_TEXTURE_MIP_COUNT] = {0};
-		int32 MipSize[MAX_TEXTURE_MIP_COUNT] = {0};
-		for (int32 MipIndex = 0; MipIndex < SlicedImage.NumMips; ++MipIndex)
-		{
-			DestMipData[MipIndex] = Texture->Source.LockMip(MipIndex);
-			MipSize[MipIndex] = Texture->Source.CalcMipSize(MipIndex) / SlicedImage.NumSlice;
-		}
-
-		for (int32 SliceIndex = 0; SliceIndex < SlicedImage.NumSlice; ++SliceIndex)
-		{
-			const uint8* SrcMipData = SlicedImage.GetMipData(0, SliceIndex);
-			for (int32 MipIndex = 0; MipIndex < SlicedImage.NumMips; ++MipIndex)
-			{
-				FMemory::Memcpy(DestMipData[MipIndex] + MipSize[MipIndex] * SliceIndex, SrcMipData, MipSize[MipIndex]);
-				SrcMipData += MipSize[MipIndex];
-			}
-		}
-
-		for (int32 MipIndex = 0; MipIndex < SlicedImage.NumMips; ++MipIndex)
-		{
-			Texture->Source.UnlockMip(MipIndex);
-		}
-
 		if (SlicedImage.MipGenSettings.IsSet())
 		{
 			// if the source has mips we keep the mips by default, unless the user changes that
@@ -387,25 +358,52 @@ namespace UE::Interchange::Private::InterchangeTextureFactory
 		}
 	}
 
-	void SetupTextureSourceData(UTextureLightProfile* TextureLightProfile, const FImportLightProfile& LightProfile)
+	void SetupTextureSourceData(UTextureLightProfile* TextureLightProfile, FImportLightProfile& LightProfile)
 	{
-		const FImportImage& ImportImage = LightProfile;
+		FImportImage& ImportImage = LightProfile;
 		SetupTextureSourceData(TextureLightProfile, ImportImage);
 
 		TextureLightProfile->Brightness = LightProfile.Brightness;
 		TextureLightProfile->TextureMultiplier = LightProfile.TextureMultiplier;
 	}
 
-	bool SetupTexture2DSourceData(UTexture2D* Texture2D, const FTexturePayloadVariant& TexturePayload)
+	bool CanSetupTexture2DSourceData(FTexturePayloadVariant& TexturePayload)
+	{
+		if (TOptional<FImportBlockedImage>* BlockedImage = TexturePayload.TryGet<TOptional<FImportBlockedImage>>())
+		{
+			if (BlockedImage->IsSet())
+			{
+				return (*BlockedImage)->IsValid();
+			}
+		}
+		else if (TOptional<FImportImage>* Image = TexturePayload.TryGet<TOptional<FImportImage>>())
+		{
+			if (Image->IsSet())
+			{
+				return (*Image)->IsValid();
+			}
+		}
+		else if (TOptional<FImportLightProfile>* LightProfile = TexturePayload.TryGet<TOptional<FImportLightProfile>>())
+		{
+			if (LightProfile->IsSet())
+			{
+				return (*LightProfile)->IsValid();
+			}
+		}
+
+		return false;
+	}
+
+	bool SetupTexture2DSourceData(UTexture2D* Texture2D, FTexturePayloadVariant& TexturePayload)
 	{
-		if (const TOptional<FImportBlockedImage>* BlockedImage = TexturePayload.TryGet<TOptional<FImportBlockedImage>>())
+		if (TOptional<FImportBlockedImage>* BlockedImage = TexturePayload.TryGet<TOptional<FImportBlockedImage>>())
 		{
 			if (BlockedImage->IsSet())
 			{
 				return SetupTexture2DSourceData(Texture2D, BlockedImage->GetValue());
 			}
 		}
-		else if (const TOptional<FImportImage>* Image = TexturePayload.TryGet<TOptional<FImportImage>>())
+		else if (TOptional<FImportImage>* Image = TexturePayload.TryGet<TOptional<FImportImage>>())
 		{
 			if (Image->IsSet())
 			{
@@ -413,32 +411,62 @@ namespace UE::Interchange::Private::InterchangeTextureFactory
 				return true;
 			}
 		}
-		else if (const TOptional<FImportLightProfile>* OptionalLightProfile = TexturePayload.TryGet<TOptional<FImportLightProfile>>())
+		else if (TOptional<FImportLightProfile>* OptionalLightProfile = TexturePayload.TryGet<TOptional<FImportLightProfile>>())
 		{
-			const FImportLightProfile& LightProfile = OptionalLightProfile->GetValue();
+			if (OptionalLightProfile->IsSet())
+			{
+				FImportLightProfile& LightProfile = OptionalLightProfile->GetValue();
+
+				if (UTextureLightProfile* TextureLightProfile = Cast<UTextureLightProfile>(Texture2D))
+				{
+					SetupTextureSourceData(TextureLightProfile, LightProfile);
+				}
+				else
+				{
+					SetupTextureSourceData(Texture2D, LightProfile);
+				}
+			}
+
+			return true;
+		}
 
-			if (UTextureLightProfile* TextureLightProfile = Cast<UTextureLightProfile>(Texture2D))
+		return false;
+	}
+
+	bool CanSetupTextureCubeSourceData(FTexturePayloadVariant& TexturePayload)
+	{
+		if (TOptional<FImportSlicedImage>* SlicedImage = TexturePayload.TryGet<TOptional<FImportSlicedImage>>())
+		{
+			if (SlicedImage->IsSet())
 			{
-				SetupTextureSourceData(TextureLightProfile, LightProfile);
+				return (*SlicedImage)->IsValid() && (*SlicedImage)->NumSlice == 6;
 			}
-			else
+		}
+		else if (TOptional<FImportImage>* Image = TexturePayload.TryGet<TOptional<FImportImage>>())
+		{
+			if (Image->IsSet())
 			{
-				SetupTextureSourceData(Texture2D, LightProfile);
+				return (*Image)->IsValid();
+			}
+		}
+		else if (TOptional<FImportLightProfile>* LightProfile = TexturePayload.TryGet<TOptional<FImportLightProfile>>())
+		{
+			if (LightProfile->IsSet())
+			{
+				return (*LightProfile)->IsValid();
 			}
-
-			return true;
 		}
 
 		return false;
 	}
 
-	bool SetupTextureCubeSourceData(UTextureCube* TextureCube, const FTexturePayloadVariant& TexturePayload)
+	bool SetupTextureCubeSourceData(UTextureCube* TextureCube, FTexturePayloadVariant& TexturePayload)
 	{
-		if (const TOptional<FImportSlicedImage>* OptionalSlicedImage = TexturePayload.TryGet<TOptional<FImportSlicedImage>>())
+		if (TOptional<FImportSlicedImage>* OptionalSlicedImage = TexturePayload.TryGet<TOptional<FImportSlicedImage>>())
 		{
 			if (OptionalSlicedImage->IsSet())
 			{
-				const FImportSlicedImage& SlicedImage = OptionalSlicedImage->GetValue();
+				FImportSlicedImage& SlicedImage = OptionalSlicedImage->GetValue();
 
 				// Cube texture always have six slice
 				if (SlicedImage.NumSlice == 6)
@@ -448,20 +476,20 @@ namespace UE::Interchange::Private::InterchangeTextureFactory
 				}
 			}
 		}
-		else if (const TOptional<FImportImage>* OptionalImage = TexturePayload.TryGet<TOptional<FImportImage>>())
+		else if (TOptional<FImportImage>* OptionalImage = TexturePayload.TryGet<TOptional<FImportImage>>())
 		{
 			if (OptionalImage->IsSet())
 			{
-				const FImportImage& Image = OptionalImage->GetValue();
+				FImportImage& Image = OptionalImage->GetValue();
 				SetupTextureSourceData(TextureCube, Image);
 				return true;
 			}
 		}
-		else if (const TOptional<FImportLightProfile>* OptionalLightProfile = TexturePayload.TryGet<TOptional<FImportLightProfile>>())
+		else if (TOptional<FImportLightProfile>* OptionalLightProfile = TexturePayload.TryGet<TOptional<FImportLightProfile>>())
 		{
 			if (OptionalLightProfile->IsSet())
 			{
-				const FImportLightProfile& LightProfile = OptionalLightProfile->GetValue();
+				FImportLightProfile& LightProfile = OptionalLightProfile->GetValue();
 				SetupTextureSourceData(TextureCube, LightProfile);
 				return true;
 			}
@@ -470,37 +498,72 @@ namespace UE::Interchange::Private::InterchangeTextureFactory
 		return false;
 	}
 
-	bool SetupTexture2DArraySourceData(UTexture2DArray* Texture2DArray, const FTexturePayloadVariant& TexturePayload)
+	bool CanSetupTexture2DArraySourceData(FTexturePayloadVariant& TexturePayload)
+	{
+		if (TOptional<FImportSlicedImage>* SlicedImage = TexturePayload.TryGet<TOptional<FImportSlicedImage>>())
+		{
+			if (SlicedImage->IsSet())
+			{
+				return (*SlicedImage)->IsValid();
+			}
+		}
+		else if (TOptional<FImportImage>* Image = TexturePayload.TryGet<TOptional<FImportImage>>())
+		{
+			if (Image->IsSet())
+			{
+				return (*Image)->IsValid();
+			}
+		}
+		else if (TOptional<FImportLightProfile>* LightProfile = TexturePayload.TryGet<TOptional<FImportLightProfile>>())
+		{
+			if (LightProfile->IsSet())
+			{
+				return (*LightProfile)->IsValid();
+			}
+		}
+
+		return false;
+	}
+
+	bool SetupTexture2DArraySourceData(UTexture2DArray* Texture2DArray, FTexturePayloadVariant& TexturePayload)
 	{
-		if (const TOptional<FImportSlicedImage>* OptionalSlicedImage = TexturePayload.TryGet<TOptional<FImportSlicedImage>>())
+		if (TOptional<FImportSlicedImage>* OptionalSlicedImage = TexturePayload.TryGet<TOptional<FImportSlicedImage>>())
 		{
 			if (OptionalSlicedImage->IsSet())
 			{
-				const FImportSlicedImage& SlicedImage = OptionalSlicedImage->GetValue();
+				FImportSlicedImage& SlicedImage = OptionalSlicedImage->GetValue();
 
 				SetupTextureSourceData(Texture2DArray, SlicedImage);
 				return true;
 			}
 		}
-		else if (const TOptional<FImportImage>* OptionalImage = TexturePayload.TryGet<TOptional<FImportImage>>())
+		else if (TOptional<FImportImage>* OptionalImage = TexturePayload.TryGet<TOptional<FImportImage>>())
 		{
 			if (OptionalImage->IsSet())
 			{
-				const FImportImage& Image = OptionalImage->GetValue();
+				FImportImage& Image = OptionalImage->GetValue();
 				SetupTextureSourceData(Texture2DArray, Image);
 				return true;
 			}
 		}
-		else if (const TOptional<FImportLightProfile>* OptionalLightProfile = TexturePayload.TryGet<TOptional<FImportLightProfile>>())
+		else if (TOptional<FImportLightProfile>* OptionalLightProfile = TexturePayload.TryGet<TOptional<FImportLightProfile>>())
 		{
-			const FImportLightProfile& LightProfile = OptionalLightProfile->GetValue();
-			SetupTextureSourceData(Texture2DArray, LightProfile);
+			if (OptionalLightProfile->IsSet())
+			{
+				FImportLightProfile& LightProfile = OptionalLightProfile->GetValue();
+				SetupTextureSourceData(Texture2DArray, LightProfile);
+			}
 			return true;
 		}
 
 		return false;
 	}
 
+	void LogErrorInvalidPayload(const FString& TextureClass, const FString& ObjectName)
+	{
+		UE_LOG(LogInterchangeImport, Error, TEXT("UInterchangeTextureFactory: The Payload was invalid for a %s. (%s)"), *TextureClass, *ObjectName);
+	}
+
 #endif // WITH_EDITORONLY_DATA
  }
 
@@ -510,7 +573,7 @@ UClass* UInterchangeTextureFactory::GetFactoryClass() const
 	return UTexture::StaticClass();
 }
 
-UObject* UInterchangeTextureFactory::CreateEmptyAsset(const FCreateAssetParams& Arguments) const
+UObject* UInterchangeTextureFactory::CreateEmptyAsset(const FCreateAssetParams& Arguments)
 {
 	using namespace  UE::Interchange::Private::InterchangeTextureFactory;
 
@@ -573,13 +636,12 @@ UObject* UInterchangeTextureFactory::CreateEmptyAsset(const FCreateAssetParams&
 		return nullptr;
 	}
 
-	Texture->PreEditChange(nullptr);
-
 #endif //WITH_EDITORONLY_DATA
 	return Texture;
 }
 
-UObject* UInterchangeTextureFactory::CreateAsset(const UInterchangeTextureFactory::FCreateAssetParams& Arguments) const
+// The payload fetching and the heavy operations are done here
+UObject* UInterchangeTextureFactory::CreateAsset(const FCreateAssetParams& Arguments)
 {
 #if !WITH_EDITORONLY_DATA
 
@@ -588,6 +650,8 @@ UObject* UInterchangeTextureFactory::CreateAsset(const UInterchangeTextureFactor
 
 #else
 
+	TRACE_CPUPROFILER_EVENT_SCOPE(UInterchangeTextureFactory::CreateAsset);
+
 	using namespace  UE::Interchange::Private::InterchangeTextureFactory;
 
 	if (!Arguments.AssetNode)
@@ -627,7 +691,7 @@ UObject* UInterchangeTextureFactory::CreateAsset(const UInterchangeTextureFactor
 		return nullptr;
 	}
 
-	FTexturePayloadVariant TexturePayload = GetTexturePayload(Arguments.SourceData, PayLoadKey.GetValue(), TextureNodeVariant, Arguments.Translator);
+	TexturePayload = GetTexturePayload(Arguments.SourceData, PayLoadKey.GetValue(), TextureNodeVariant, Arguments.Translator);
 
 	if(TexturePayload.IsType<FEmptyVariantState>())
 	{
@@ -660,59 +724,32 @@ UObject* UInterchangeTextureFactory::CreateAsset(const UInterchangeTextureFactor
 	}
 
 
-	// Setup source data
+	bool bCanSetup = false;
+	// Check if the payload is valid for the Texture
 	if (UTexture2D* Texture2D = Cast<UTexture2D>(Texture))
 	{
-		if (!SetupTexture2DSourceData(Texture2D, TexturePayload))
-		{
-			UE_LOG(LogInterchangeImport, Error, TEXT("UInterchangeTextureFactory: The Payload was invalid for a %s. (%s)"), *(TextureClass->GetName()), *Arguments.AssetName);
-			return nullptr;
-		}
+		bCanSetup = CanSetupTexture2DSourceData(TexturePayload);
 	}
 	else if (UTextureCube* TextureCube = Cast<UTextureCube>(Texture))
 	{
-		if (!SetupTextureCubeSourceData(TextureCube, TexturePayload))
-		{
-			UE_LOG(LogInterchangeImport, Error, TEXT("UInterchangeTextureFactory: The Payload was invalid for a TextureCube. (%s)"), *Arguments.AssetName);
-			return nullptr;
-		}
+		bCanSetup = CanSetupTextureCubeSourceData(TexturePayload);
 	}
 	else if (UTexture2DArray* Texture2DArray = Cast<UTexture2DArray>(Texture))
 	{
-		if (!SetupTexture2DArraySourceData(Texture2DArray, TexturePayload))
-		{
-			UE_LOG(LogInterchangeImport, Error, TEXT("UInterchangeTextureFactory: The Payload was invalid for a TextureArray. (%s)"), *Arguments.AssetName);
-			return nullptr;
-		}
-	}
-	else
-	{
-		// The texture is not supported
-		Texture->RemoveFromRoot();
-		Texture->MarkPendingKill();
-		return nullptr;
+		bCanSetup = CanSetupTexture2DArraySourceData(TexturePayload);
 	}
 
-	UInterchangeBaseNode* TextureFactoryNode = Arguments.AssetNode;
-	if(!Arguments.ReimportObject)
+	if (!bCanSetup)
 	{
-		/** Apply all TextureNode custom attributes to the texture asset */
-		TextureFactoryNode->ApplyAllCustomAttributeToAsset(Texture);
-	}
-	else
-	{
-		UInterchangeAssetImportData* InterchangeAssetImportData = Cast<UInterchangeAssetImportData>(Texture->AssetImportData);
-		UInterchangeBaseNode* PreviousNode = nullptr;
-		if (InterchangeAssetImportData)
+		LogErrorInvalidPayload(Texture->GetClass()->GetName(), Texture->GetName());
+
+		// The texture is not supported
+		if (!Arguments.ReimportObject)
 		{
-			PreviousNode = InterchangeAssetImportData->NodeContainer->GetNode(InterchangeAssetImportData->NodeUniqueID);
+			Texture->RemoveFromRoot();
+			Texture->MarkPendingKill();
 		}
-
-		UInterchangeBaseNode* CurrentNode = NewObject<UInterchangeBaseNode>(GetTransientPackage(), SupportedFactoryNodeClass);
-		UInterchangeBaseNode::CopyStorage(TextureFactoryNode, CurrentNode);
-		CurrentNode->FillAllCustomAttributeFromAsset(Texture);
-		//Apply reimport strategy
-		UE::Interchange::FFactoryCommon::ApplyReimportStrategyToAsset(Arguments.ReimportStrategyFlags, Texture, PreviousNode, CurrentNode, TextureFactoryNode);
+		return nullptr;
 	}
 
 	//Getting the file Hash will cache it into the source data
@@ -725,24 +762,93 @@ UObject* UInterchangeTextureFactory::CreateAsset(const UInterchangeTextureFactor
 }
 
 /* This function is call in the completion task on the main thread, use it to call main thread post creation step for your assets*/
-void UInterchangeTextureFactory::PreImportPreCompletedCallback(const FImportPreCompletedCallbackParams& Arguments) const
+void UInterchangeTextureFactory::PreImportPreCompletedCallback(const FImportPreCompletedCallbackParams& Arguments)
 {
+	TRACE_CPUPROFILER_EVENT_SCOPE(UInterchangeTextureFactory::PreImportPreCompletedCallback);
+
 	check(IsInGameThread());
+
+	UTexture* Texture = Cast<UTexture>(Arguments.ImportedObject);
+
+#if WITH_EDITOR
+
+	// Finish the import on the game thread by doing the setup on the texture here
+	if (Texture)
+	{
+		Texture->PreEditChange(nullptr);
+
+		using namespace UE::Interchange::Private::InterchangeTextureFactory;
+
+		// Setup source data
+		if (UTexture2D* Texture2D = Cast<UTexture2D>(Texture))
+		{
+			if (!SetupTexture2DSourceData(Texture2D, TexturePayload))
+			{
+				LogErrorInvalidPayload(Arguments.ImportedObject->GetClass()->GetName(), Arguments.ImportedObject->GetName());
+				return;			
+			}
+		}
+		else if (UTextureCube* TextureCube = Cast<UTextureCube>(Texture))
+		{
+			if (!SetupTextureCubeSourceData(TextureCube, TexturePayload))
+			{
+				LogErrorInvalidPayload(Arguments.ImportedObject->GetClass()->GetName(), Arguments.ImportedObject->GetName());
+				return;
+			}
+		}
+		else if (UTexture2DArray* Texture2DArray = Cast<UTexture2DArray>(Arguments.ImportedObject))
+		{
+			if (!SetupTexture2DArraySourceData(Texture2DArray, TexturePayload))
+			{
+				LogErrorInvalidPayload(Arguments.ImportedObject->GetClass()->GetName(), Arguments.ImportedObject->GetName());
+				return;
+			}
+		}
+		else
+		{
+			// This should never happen.
+			ensure(false);
+		}
+
+
+		UInterchangeBaseNode* TextureFactoryNode = Arguments.FactoryNode;
+		if (!Arguments.bIsReimport)
+		{
+			/** Apply all TextureNode custom attributes to the texture asset */
+			TextureFactoryNode->ApplyAllCustomAttributeToAsset(Texture);
+		}
+		else
+		{
+			UInterchangeAssetImportData* InterchangeAssetImportData = Cast<UInterchangeAssetImportData>(Texture->AssetImportData);
+			UInterchangeBaseNode* PreviousNode = nullptr;
+			if (InterchangeAssetImportData)
+			{
+				PreviousNode = InterchangeAssetImportData->NodeContainer->GetNode(InterchangeAssetImportData->NodeUniqueID);
+			}
+
+			UInterchangeBaseNode* CurrentNode = NewObject<UInterchangeBaseNode>(GetTransientPackage(), GetSupportedFactoryNodeClass(TextureFactoryNode));
+			UInterchangeBaseNode::CopyStorage(TextureFactoryNode, CurrentNode);
+			CurrentNode->FillAllCustomAttributeFromAsset(Texture);
+			//Apply reimport strategy
+			UE::Interchange::FFactoryCommon::ApplyReimportStrategyToAsset(Arguments.ReimportStrategyFlags, Texture, PreviousNode, CurrentNode, TextureFactoryNode);
+		}
+	}
+#endif //WITH_EDITOR
+
 	Super::PreImportPreCompletedCallback(Arguments);
+
 	//TODO make sure this work at runtime
 #if WITH_EDITORONLY_DATA
-	if (ensure(Arguments.ImportedObject && Arguments.SourceData))
+	if (ensure(Texture && Arguments.SourceData))
 	{
 		//We must call the Update of the asset source file in the main thread because UAssetImportData::Update execute some delegate we do not control
-		UTexture* ImportedTexture = CastChecked<UTexture>(Arguments.ImportedObject);
-		
-		UE::Interchange::FFactoryCommon::FUpdateImportAssetDataParameters UpdateImportAssetDataParameters(ImportedTexture
-																						 , ImportedTexture->AssetImportData
+		UE::Interchange::FFactoryCommon::FUpdateImportAssetDataParameters UpdateImportAssetDataParameters(Texture
+																						 , Texture->AssetImportData
 																						 , Arguments.SourceData
 																						 , Arguments.NodeUniqueID
 																						 , Arguments.NodeContainer
 																						 , Arguments.Pipelines);
-		ImportedTexture->AssetImportData = UE::Interchange::FFactoryCommon::UpdateImportAssetData(UpdateImportAssetDataParameters);
+		Texture->AssetImportData = UE::Interchange::FFactoryCommon::UpdateImportAssetData(UpdateImportAssetDataParameters);
 	}
 #endif
 }
\ No newline at end of file
diff --git a/Engine/Plugins/Runtime/Interchange/Source/Import/Private/Texture/InterchangeTexturePayloadData.cpp b/Engine/Plugins/Runtime/Interchange/Source/Import/Private/Texture/InterchangeTexturePayloadData.cpp
--- a/Engine/Plugins/Runtime/Interchange/Source/Import/Private/Texture/InterchangeTexturePayloadData.cpp
+++ b/Engine/Plugins/Runtime/Interchange/Source/Import/Private/Texture/InterchangeTexturePayloadData.cpp
@@ -50,46 +50,37 @@ bool UE::Interchange::FImportImageHelper::IsImportResolutionValid(int32 Width, i
 }
 
 void UE::Interchange::FImportImage::Init2DWithParams(int32 InSizeX, int32 InSizeY, ETextureSourceFormat InFormat, bool InSRGB)
+{
+	Init2DWithParams(InSizeX, InSizeY, 1, InFormat, InSRGB);
+}
+
+void UE::Interchange::FImportImage::Init2DWithParams(int32 InSizeX, int32 InSizeY, int32 InNumMips, ETextureSourceFormat InFormat, bool InSRGB)
 {
 	SizeX = InSizeX;
 	SizeY = InSizeY;
-	NumMips = 1;
+	NumMips = InNumMips;
 	Format = InFormat;
-	SRGB = InSRGB;
+	bSRGB = InSRGB;
+	RawData = FUniqueBuffer::Alloc(ComputeBufferSize());
 }
 
 void UE::Interchange::FImportImage::Init2DWithOneMip(int32 InSizeX, int32 InSizeY, ETextureSourceFormat InFormat, const void* InData)
 {
-#if WITH_EDITOR
-	SizeX = InSizeX;
-	SizeY = InSizeY;
-	NumMips = 1;
-	Format = InFormat;
-	RawData.AddUninitialized((int64)SizeX * SizeY * FTextureSource::GetBytesPerPixel(Format));
+	Init2DWithParams(InSizeX, InSizeY, 1, InFormat, bSRGB);
+
 	if (InData)
 	{
-		FMemory::Memcpy(RawData.GetData(), InData, RawData.Num());
+		FMemory::Memcpy(RawData.GetData(), InData, RawData.GetSize());
 	}
-#endif
 }
 
 void UE::Interchange::FImportImage::Init2DWithMips(int32 InSizeX, int32 InSizeY, int32 InNumMips, ETextureSourceFormat InFormat, const void* InData)
 {
-	SizeX = InSizeX;
-	SizeY = InSizeY;
-	NumMips = InNumMips;
-	Format = InFormat;
-
-	int64 TotalSize = 0;
-	for (int32 MipIndex = 0; MipIndex < InNumMips; ++MipIndex)
-	{
-		TotalSize += GetMipSize(MipIndex);
-	}
-	RawData.AddUninitialized(TotalSize);
+	Init2DWithParams(InSizeX, InSizeY, 1, InFormat, bSRGB);
 
 	if (InData)
 	{
-		FMemory::Memcpy(RawData.GetData(), InData, RawData.Num());
+		FMemory::Memcpy(RawData.GetData(), InData, RawData.GetSize());
 	}
 }
 
@@ -106,12 +97,27 @@ int64 UE::Interchange::FImportImage::GetMipSize(int32 InMipIndex) const
 #endif
 }
 
-void* UE::Interchange::FImportImage::GetMipData(int32 InMipIndex)
+int64 UE::Interchange::FImportImage::ComputeBufferSize() const
 {
-	int64 Offset = 0;
-	for (int32 MipIndex = 0; MipIndex < InMipIndex; ++MipIndex)
+	int64 TotalSize = 0;
+	for (int32 MipIndex = 0; MipIndex < NumMips; ++MipIndex)
 	{
-		Offset += GetMipSize(MipIndex);
+		TotalSize += GetMipSize(MipIndex);
 	}
-	return &RawData[Offset];
+
+	return TotalSize;
+}
+
+TArrayView64<uint8> UE::Interchange::FImportImage::GetArrayViewOfRawData()
+{
+	return TArrayView64<uint8>(static_cast<uint8*>(RawData.GetData()), RawData.GetSize());
+}
+
+bool UE::Interchange::FImportImage::IsValid() const
+{
+	return SizeX > 0
+		&& SizeY > 0
+		&& NumMips > 0
+		&& Format != TSF_Invalid
+		&& ComputeBufferSize() == RawData.GetSize();
 }
diff --git a/Engine/Plugins/Runtime/Interchange/Source/Import/Private/Texture/InterchangeUDIMTranslator.cpp b/Engine/Plugins/Runtime/Interchange/Source/Import/Private/Texture/InterchangeUDIMTranslator.cpp
--- a/Engine/Plugins/Runtime/Interchange/Source/Import/Private/Texture/InterchangeUDIMTranslator.cpp
+++ b/Engine/Plugins/Runtime/Interchange/Source/Import/Private/Texture/InterchangeUDIMTranslator.cpp
@@ -4,9 +4,9 @@
 
 #include "InterchangeManager.h"
 #include "InterchangeTexture2DNode.h"
+#include "Texture/InterchangeTexturePayloadData.h"
 #include "Texture/InterchangeTexturePayloadInterface.h"
 
-
 #include "UDIMUtilities.h"
 #include "Async/ParallelFor.h"
 
@@ -228,8 +228,6 @@ bool UInterchangeUDIMTranslator::Translate(UInterchangeBaseNodeContainer& BaseNo
 
 TOptional<UE::Interchange::FImportBlockedImage> UInterchangeUDIMTranslator::GetBlockedTexturePayloadData(const TMap<int32, FString>& InSourcesData,  const UInterchangeSourceData* OriginalFile) const
 {
-	const FString CurrentFilename = OriginalFile->GetFilename();
-
 	UInterchangeManager& InterchangeManager = UInterchangeManager::GetInterchangeManager();
 
 	if (UInterchangeTranslatorBase* Translator = InterchangeManager.GetTranslatorSupportingPayloadInterfaceForSourceData(OriginalFile, UInterchangeTexturePayloadInterface::StaticClass()))
@@ -247,12 +245,15 @@ TOptional<UE::Interchange::FImportBlockedImage> UInterchangeUDIMTranslator::GetB
 			using namespace UE::Interchange::Private::InterchangeUDIMTranslator;
 			FScopedTranslatorsAndSourceData TranslatorsAndSourceData(*Translator, UDIMsAndSourcesFileArray.Num());
 
-			UE::Interchange::FImportBlockedImage BlockedImage;
-			BlockedImage.InitZeroed(UDIMsAndSourcesFileArray.Num());
-
-			TArray<int32> InvalidTextureIndex;
+			/**
+			 * Possible improvement notes.
+			 * If we are able at some point to extract the size and format of the textures from the translate step for some formats,
+			 * We could use those information to init the blocked image RawData and set the ImportsImages RawData to be a view into the blocked image.
+			 */ 
+			TArray<UE::Interchange::FImportImage> Images;
+			Images.AddDefaulted(UDIMsAndSourcesFileArray.Num());
 
-			ParallelFor(UDIMsAndSourcesFileArray.Num(), [&TranslatorsAndSourceData, &UDIMsAndSourcesFileArray, &BlockedImage, &InvalidTextureIndex](int32 Index)
+			ParallelFor(UDIMsAndSourcesFileArray.Num(), [&TranslatorsAndSourceData, &UDIMsAndSourcesFileArray, &Images](int32 Index)
 				{
 					TPair<UInterchangeTranslatorBase*,UInterchangeSourceData*>& TranslatorAndSourceData = TranslatorsAndSourceData[Index];
 					UInterchangeTranslatorBase* Translator = TranslatorAndSourceData.Key;
@@ -266,6 +267,7 @@ TOptional<UE::Interchange::FImportBlockedImage> UInterchangeUDIMTranslator::GetB
 					TOptional<UE::Interchange::FImportImage> Payload;
 					if (Translator->CanImportSourceData(SourceDataForBlock))
 					{
+						Translator->SourceData = SourceDataForBlock;
 						Payload = TextureTranslator->GetTexturePayloadData(SourceDataForBlock, SourceDataForBlock->GetFilename());
 					}
 
@@ -273,35 +275,40 @@ TOptional<UE::Interchange::FImportBlockedImage> UInterchangeUDIMTranslator::GetB
 					{
 						UE::Interchange::FImportImage& Image = Payload.GetValue();
 						// Consume the image
-						BlockedImage.ImagesData[Index] = MoveTemp(Image);
-
-						FTextureSourceBlock& Block =  BlockedImage.BlocksData[Index];
-						Block.BlockX = (CurrentUDIM - 1001) % 10;
-						Block.BlockY = (CurrentUDIM - 1001) / 10;
-						Block.SizeX = Image.SizeX;
-						Block.SizeY = Image.SizeY;
-						Block.NumSlices = 1;
-						Block.NumMips = Image.NumMips;
+						Images[Index] = MoveTemp(Image);
 					}
 					else
 					{
-						// Todo Interchange Log error
-						InvalidTextureIndex.Add(Index);
+						// Todo Capture the error message from the translator?
 					}
 
-					// Let the translator release it data.
+					// Let the translator release his data.
 					Translator->ImportFinish();
 				}
 				, EParallelForFlags::Unbalanced | EParallelForFlags::BackgroundPriority);
 
-			for (int32 Index : InvalidTextureIndex)
-			{
-				BlockedImage.RemoveBlock(Index);
-			}
 
-			if (BlockedImage.HasData())
+			UE::Interchange::FImportBlockedImage BlockedImage;
+			// If the image that triggered the import is not valid. We shouldn't import the rest of the images into the UDIM.
+			if (BlockedImage.InitDataSharedAmongBlocks(Images[0]))
 			{
-				return BlockedImage;
+				BlockedImage.BlocksData.Reserve(Images.Num());
+
+				for (int32 Index = 0; Index < Images.Num(); ++Index)
+				{
+					int32 UDIMIndex = UDIMsAndSourcesFileArray[Index]->Key;
+					const int32 BlockX = (UDIMIndex - 1001) % 10;
+					const int32 BlockY = (UDIMIndex - 1001) / 10;
+					BlockedImage.InitBlockFromImage(BlockX, BlockY, Images[Index]);
+				}
+
+				BlockedImage.MigrateDataFromImagesToRawData(Images);
+
+				if (BlockedImage.BlocksData.Num() > 0)
+				{
+					return BlockedImage;
+				}
+				
 			}
 		}
 
@@ -311,11 +318,30 @@ TOptional<UE::Interchange::FImportBlockedImage> UInterchangeUDIMTranslator::GetB
 			TOptional<UE::Interchange::FImportImage> Payload = TextureTranslator->GetTexturePayloadData(OriginalFile, OriginalFile->GetFilename());
 			if (Payload.IsSet())
 			{
+				UE::Interchange::FImportImage& Image = Payload.GetValue();
+
 				UE::Interchange::FImportBlockedImage BlockedImage;
-				// Consume the image
-				BlockedImage.ImagesData.Add(MoveTemp(Payload.GetValue()));
-				Translator->ImportFinish();
-				return BlockedImage;
+				if (BlockedImage.InitDataSharedAmongBlocks(Image))
+				{
+					const FString FilenameNoExtension = FPaths::GetBaseFilename(OriginalFile->GetFilename());
+
+					FString PreUDIMName;
+					FString PostUDIMName;
+					const int32 UDIMIndex = UE::TextureUtilitiesCommon::ParseUDIMName(FilenameNoExtension, UdimRegexPattern, PreUDIMName, PostUDIMName);
+					
+					const int32 BlockX = (UDIMIndex - 1001) % 10;
+					const int32 BlockY = (UDIMIndex - 1001) / 10;
+
+					BlockedImage.InitBlockFromImage(BlockX, BlockY, Image);
+
+					// Consume the payload
+					BlockedImage.RawData = MoveTemp(Image.RawData);
+		
+					Translator->ImportFinish();
+					return BlockedImage;
+				}
+
+				// Todo else Grab error message from translator
 			}
 		}
 
diff --git a/Engine/Plugins/Runtime/Interchange/Source/Import/Public/Material/InterchangeMaterialFactory.h b/Engine/Plugins/Runtime/Interchange/Source/Import/Public/Material/InterchangeMaterialFactory.h
--- a/Engine/Plugins/Runtime/Interchange/Source/Import/Public/Material/InterchangeMaterialFactory.h
+++ b/Engine/Plugins/Runtime/Interchange/Source/Import/Public/Material/InterchangeMaterialFactory.h
@@ -19,9 +19,9 @@ public:
 	// Interchange factory base interface begin
 
 	virtual UClass* GetFactoryClass() const override;
-	virtual UObject* CreateEmptyAsset(const FCreateAssetParams& Arguments) const override;
-	virtual UObject* CreateAsset(const FCreateAssetParams& Arguments) const override;
-	virtual void PreImportPreCompletedCallback(const FImportPreCompletedCallbackParams& Arguments) const override;
+	virtual UObject* CreateEmptyAsset(const FCreateAssetParams& Arguments) override;
+	virtual UObject* CreateAsset(const FCreateAssetParams& Arguments) override;
+	virtual void PreImportPreCompletedCallback(const FImportPreCompletedCallbackParams& Arguments) override;
 
 	// Interchange factory base interface end
 	//////////////////////////////////////////////////////////////////////////
diff --git a/Engine/Plugins/Runtime/Interchange/Source/Import/Public/Mesh/InterchangePhysicsAssetFactory.h b/Engine/Plugins/Runtime/Interchange/Source/Import/Public/Mesh/InterchangePhysicsAssetFactory.h
--- a/Engine/Plugins/Runtime/Interchange/Source/Import/Public/Mesh/InterchangePhysicsAssetFactory.h
+++ b/Engine/Plugins/Runtime/Interchange/Source/Import/Public/Mesh/InterchangePhysicsAssetFactory.h
@@ -19,8 +19,8 @@ public:
 	// Interchange factory base interface begin
 
 	virtual UClass* GetFactoryClass() const override;
-	virtual UObject* CreateEmptyAsset(const FCreateAssetParams& Arguments) const override;
-	virtual UObject* CreateAsset(const FCreateAssetParams& Arguments) const override;
+	virtual UObject* CreateEmptyAsset(const FCreateAssetParams& Arguments) override;
+	virtual UObject* CreateAsset(const FCreateAssetParams& Arguments) override;
 	//virtual void PreImportPreCompletedCallback(const FImportPreCompletedCallbackParams& Arguments) const override;
 	//virtual void PostImportPreCompletedCallback(const FImportPreCompletedCallbackParams& Arguments) const override;
 
diff --git a/Engine/Plugins/Runtime/Interchange/Source/Import/Public/Mesh/InterchangeSkeletalMeshFactory.h b/Engine/Plugins/Runtime/Interchange/Source/Import/Public/Mesh/InterchangeSkeletalMeshFactory.h
--- a/Engine/Plugins/Runtime/Interchange/Source/Import/Public/Mesh/InterchangeSkeletalMeshFactory.h
+++ b/Engine/Plugins/Runtime/Interchange/Source/Import/Public/Mesh/InterchangeSkeletalMeshFactory.h
@@ -21,9 +21,9 @@ public:
 	// Interchange factory base interface begin
 
 	virtual UClass* GetFactoryClass() const override;
-	virtual UObject* CreateEmptyAsset(const FCreateAssetParams& Arguments) const override;
-	virtual UObject* CreateAsset(const FCreateAssetParams& Arguments) const override;
-	virtual void PreImportPreCompletedCallback(const FImportPreCompletedCallbackParams& Arguments) const override;
+	virtual UObject* CreateEmptyAsset(const FCreateAssetParams& Arguments) override;
+	virtual UObject* CreateAsset(const FCreateAssetParams& Arguments) override;
+	virtual void PreImportPreCompletedCallback(const FImportPreCompletedCallbackParams& Arguments) override;
 	//virtual void PostImportPreCompletedCallback(const FImportPreCompletedCallbackParams& Arguments) const override;
 
 	// Interchange factory base interface end
diff --git a/Engine/Plugins/Runtime/Interchange/Source/Import/Public/Mesh/InterchangeSkeletonFactory.h b/Engine/Plugins/Runtime/Interchange/Source/Import/Public/Mesh/InterchangeSkeletonFactory.h
--- a/Engine/Plugins/Runtime/Interchange/Source/Import/Public/Mesh/InterchangeSkeletonFactory.h
+++ b/Engine/Plugins/Runtime/Interchange/Source/Import/Public/Mesh/InterchangeSkeletonFactory.h
@@ -19,8 +19,8 @@ public:
 	// Interchange factory base interface begin
 
 	virtual UClass* GetFactoryClass() const override;
-	virtual UObject* CreateEmptyAsset(const FCreateAssetParams& Arguments) const override;
-	virtual UObject* CreateAsset(const FCreateAssetParams& Arguments) const override;
+	virtual UObject* CreateEmptyAsset(const FCreateAssetParams& Arguments) override;
+	virtual UObject* CreateAsset(const FCreateAssetParams& Arguments) override;
 	//virtual void PreImportPreCompletedCallback(const FImportPreCompletedCallbackParams& Arguments) const override;
 
 	// Interchange factory base interface end
diff --git a/Engine/Plugins/Runtime/Interchange/Source/Import/Public/Mesh/InterchangeStaticMeshFactory.h b/Engine/Plugins/Runtime/Interchange/Source/Import/Public/Mesh/InterchangeStaticMeshFactory.h
--- a/Engine/Plugins/Runtime/Interchange/Source/Import/Public/Mesh/InterchangeStaticMeshFactory.h
+++ b/Engine/Plugins/Runtime/Interchange/Source/Import/Public/Mesh/InterchangeStaticMeshFactory.h
@@ -19,9 +19,9 @@ public:
 	// Interchange factory base interface begin
 
 	virtual UClass* GetFactoryClass() const override;
-	virtual UObject* CreateEmptyAsset(const FCreateAssetParams& Arguments) const override;
-	virtual UObject* CreateAsset(const FCreateAssetParams& Arguments) const override;
-	virtual void PreImportPreCompletedCallback(const FImportPreCompletedCallbackParams& Arguments) const override;
+	virtual UObject* CreateEmptyAsset(const FCreateAssetParams& Arguments) override;
+	virtual UObject* CreateAsset(const FCreateAssetParams& Arguments) override;
+	virtual void PreImportPreCompletedCallback(const FImportPreCompletedCallbackParams& Arguments) override;
 	//virtual void PostImportPreCompletedCallback(const FImportPreCompletedCallbackParams& Arguments) const override;
 
 	// Interchange factory base interface end
diff --git a/Engine/Plugins/Runtime/Interchange/Source/Import/Public/Texture/InterchangeBlockedTexturePayloadData.h b/Engine/Plugins/Runtime/Interchange/Source/Import/Public/Texture/InterchangeBlockedTexturePayloadData.h
--- a/Engine/Plugins/Runtime/Interchange/Source/Import/Public/Texture/InterchangeBlockedTexturePayloadData.h
+++ b/Engine/Plugins/Runtime/Interchange/Source/Import/Public/Texture/InterchangeBlockedTexturePayloadData.h
@@ -6,6 +6,7 @@
 
 #include "CoreMinimal.h"
 #include "Engine/Texture.h"
+#include "Memory/SharedBuffer.h"
 
 namespace UE
 {
@@ -14,33 +15,30 @@ namespace UE
 		// Also known as a UDIMs texture
 		struct FImportBlockedImage
 		{
-			// The first image data will be used for the default compression setting 
-			// The first image will also be use for the pixel format and gamma (bsRGB)
-
+			FUniqueBuffer RawData;
 			TArray<FTextureSourceBlock> BlocksData;
-			TArray<FImportImage> ImagesData;
-
-			void InitZeroed(uint32 Count)
-			{
-				BlocksData.AddZeroed(Count);
-				ImagesData.AddZeroed(Count);
-			}
-
-			void RemoveBlock(int32 Index)
-			{
-				BlocksData.RemoveAtSwap(Index);
-				ImagesData.RemoveAtSwap(Index);
-			}
-
-			bool HasData() const
-			{
-				return ImagesData.Num() > 0;
-			}
-
-			bool IsBlockedData() const
-			{
-				return BlocksData.Num() > 0;
-			}
+
+			ETextureSourceFormat Format = TSF_Invalid;
+			TextureCompressionSettings CompressionSettings = TC_Default;
+			bool bSRGB = true;
+			TOptional<TextureMipGenSettings> MipGenSettings;
+
+
+			int64 ComputeBufferSize() const;
+			bool IsValid() const;
+
+
+			bool InitDataSharedAmongBlocks(const FImportImage& Image);
+
+			/**
+			 * Everything except the RawData and BlocksData must be initialized before calling this function
+			 */
+			bool InitBlockFromImage(int32 BlockX, int32 BlockY, const FImportImage& Image);
+
+			/**
+			 * Everything except the RawData must be initialized before calling this function
+			 */
+			bool MigrateDataFromImagesToRawData(TArray<FImportImage>& Images);
 		};
 	}
 }
\ No newline at end of file
diff --git a/Engine/Plugins/Runtime/Interchange/Source/Import/Public/Texture/InterchangeSlicedTexturePayloadData.h b/Engine/Plugins/Runtime/Interchange/Source/Import/Public/Texture/InterchangeSlicedTexturePayloadData.h
--- a/Engine/Plugins/Runtime/Interchange/Source/Import/Public/Texture/InterchangeSlicedTexturePayloadData.h
+++ b/Engine/Plugins/Runtime/Interchange/Source/Import/Public/Texture/InterchangeSlicedTexturePayloadData.h
@@ -3,31 +3,24 @@
 #pragma once
 
 #include "CoreMinimal.h"
-#include "Engine/Texture.h"
+#include "Texture/InterchangeTexturePayloadData.h"
 
 namespace UE
 {
 	namespace Interchange
 	{
-		struct INTERCHANGEIMPORT_API FImportSlicedImage
+		struct INTERCHANGEIMPORT_API FImportSlicedImage : public FImportImage
 		{
-			TArray64<uint8> RawData;
-			ETextureSourceFormat Format = TSF_Invalid;
-			TextureCompressionSettings CompressionSettings = TC_Default;
-			int32 NumMips = 0;
-			int32 SizeX = 0;
-			int32 SizeY = 0;
 			int32 NumSlice = 0;
-			bool bSRGB = true;
-			TOptional<TextureMipGenSettings> MipGenSettings;
 
 			void Init(int32 InSizeX, int32 InSizeY, int32 InNumSlice, int32 InNumMips, ETextureSourceFormat InFormat , bool InSRGB);
 
-			/** Return the size of a mip for one slice */
-			int64 GetMipSize(int32 InMipIndex) const;
-
 			const uint8* GetMipData(int32 InMipIndex, int32 InSliceIndex = INDEX_NONE) const;
 			uint8* GetMipData(int32 InMipIndex, int32 InSliceIndex = INDEX_NONE);
+
+			virtual int64 ComputeBufferSize() const override;
+
+			virtual bool IsValid() const override;
 		};
 	}
 }
diff --git a/Engine/Plugins/Runtime/Interchange/Source/Import/Public/Texture/InterchangeTextureFactory.h b/Engine/Plugins/Runtime/Interchange/Source/Import/Public/Texture/InterchangeTextureFactory.h
--- a/Engine/Plugins/Runtime/Interchange/Source/Import/Public/Texture/InterchangeTextureFactory.h
+++ b/Engine/Plugins/Runtime/Interchange/Source/Import/Public/Texture/InterchangeTextureFactory.h
@@ -4,12 +4,26 @@
 
 #include "CoreMinimal.h"
 #include "InterchangeFactoryBase.h"
+#include "Misc/TVariant.h"
+#include "Texture/InterchangeBlockedTexturePayloadData.h"
+#include "Texture/InterchangeSlicedTexturePayloadData.h"
+#include "Texture/InterchangeTextureLightProfilePayloadData.h"
+#include "Texture/InterchangeTextureLightProfilePayloadData.h"
+#include "Texture/InterchangeTexturePayloadData.h"
 #include "UObject/Object.h"
 #include "UObject/ObjectMacros.h"
-#include "Texture/InterchangeTexturePayloadData.h"
 
 #include "InterchangeTextureFactory.generated.h"
 
+namespace UE::Interchange::Private::InterchangeTextureFactory
+{
+	using FTexturePayloadVariant = TVariant<FEmptyVariantState
+		, TOptional<FImportImage>
+		, TOptional<FImportBlockedImage>
+		, TOptional<FImportSlicedImage>
+		, TOptional<FImportLightProfile>>;
+}
+
 UCLASS(BlueprintType)
 class INTERCHANGEIMPORT_API UInterchangeTextureFactory : public UInterchangeFactoryBase
 {
@@ -20,11 +34,14 @@ public:
 	// Interchange factory base interface begin
 
 	virtual UClass* GetFactoryClass() const override;
-	virtual UObject* CreateEmptyAsset(const FCreateAssetParams& Arguments) const override;
-	virtual UObject* CreateAsset(const FCreateAssetParams& Arguments) const override;
-	virtual void PreImportPreCompletedCallback(const FImportPreCompletedCallbackParams& Arguments) const override;
+	virtual UObject* CreateEmptyAsset(const FCreateAssetParams& Arguments) override;
+	virtual UObject* CreateAsset(const FCreateAssetParams& Arguments) override;
+	virtual void PreImportPreCompletedCallback(const FImportPreCompletedCallbackParams& Arguments) override;
 	// Interchange factory base interface end
 	//////////////////////////////////////////////////////////////////////////
+
+private:
+	UE::Interchange::Private::InterchangeTextureFactory::FTexturePayloadVariant TexturePayload;
 };
 
 
diff --git a/Engine/Plugins/Runtime/Interchange/Source/Import/Public/Texture/InterchangeTexturePayloadData.h b/Engine/Plugins/Runtime/Interchange/Source/Import/Public/Texture/InterchangeTexturePayloadData.h
--- a/Engine/Plugins/Runtime/Interchange/Source/Import/Public/Texture/InterchangeTexturePayloadData.h
+++ b/Engine/Plugins/Runtime/Interchange/Source/Import/Public/Texture/InterchangeTexturePayloadData.h
@@ -2,8 +2,10 @@
 
 #pragma once
 
+#include "Containers/ArrayView.h"
 #include "CoreMinimal.h"
 #include "Engine/Texture.h"
+#include "Memory/SharedBuffer.h"
 
 namespace UE
 {
@@ -11,21 +13,33 @@ namespace UE
 	{
 		struct INTERCHANGEIMPORT_API FImportImage
 		{
-			TArray64<uint8> RawData;
+			virtual ~FImportImage() = default;
+			FImportImage() = default;
+			FImportImage(FImportImage&&) = default;
+			FImportImage& operator=(FImportImage&&) = default;
+
+			FImportImage(const FImportImage&) = delete;
+			FImportImage& operator=(const FImportImage&) = delete;
+
+			FUniqueBuffer RawData;
 			ETextureSourceFormat Format = TSF_Invalid;
 			TextureCompressionSettings CompressionSettings = TC_Default;
 			int32 NumMips = 0;
 			int32 SizeX = 0;
 			int32 SizeY = 0;
-			bool SRGB = true;
+			bool bSRGB = true;
 			TOptional<TextureMipGenSettings> MipGenSettings;
 
 			void Init2DWithParams(int32 InSizeX, int32 InSizeY, ETextureSourceFormat InFormat, bool InSRGB);
+			void Init2DWithParams(int32 InSizeX, int32 InSizeY, int32 InNumMips, ETextureSourceFormat InFormat, bool InSRGB);
 			void Init2DWithOneMip(int32 InSizeX, int32 InSizeY, ETextureSourceFormat InFormat, const void* InData = nullptr);
 			void Init2DWithMips(int32 InSizeX, int32 InSizeY, int32 InNumMips, ETextureSourceFormat InFormat, const void* InData = nullptr);
 
 			int64 GetMipSize(int32 InMipIndex) const;
-			void* GetMipData(int32 InMipIndex);
+			virtual int64 ComputeBufferSize() const;
+
+			TArrayView64<uint8> GetArrayViewOfRawData();
+			virtual bool IsValid() const;
 		};
 
 		struct INTERCHANGEIMPORT_API FImportImageHelper
diff --git a/Engine/Source/Runtime/Engine/Classes/Engine/Texture.h b/Engine/Source/Runtime/Engine/Classes/Engine/Texture.h
--- a/Engine/Source/Runtime/Engine/Classes/Engine/Texture.h
+++ b/Engine/Source/Runtime/Engine/Classes/Engine/Texture.h
@@ -134,6 +134,12 @@ struct FTextureSource
 		int32 InNumBlocks,
 		const uint8** InDataPerBlock);
 
+	ENGINE_API void InitBlocked(const ETextureSourceFormat* InLayerFormats,
+		const FTextureSourceBlock* InBlocks,
+		int32 InNumLayers,
+		int32 InNumBlocks,
+		FSharedBuffer NewData);
+
 	ENGINE_API void InitLayered(
 		int32 NewSizeX,
 		int32 NewSizeY,
@@ -144,6 +150,16 @@ struct FTextureSource
 		const uint8* NewData = NULL
 		);
 
+	ENGINE_API void InitLayered(
+		int32 NewSizeX,
+		int32 NewSizeY,
+		int32 NewNumSlices,
+		int32 NewNumLayers,
+		int32 NewNumMips,
+		const ETextureSourceFormat* NewLayerFormat,
+		FSharedBuffer NewData
+	);
+
 	/**
 	 * Initialize the source data with the given size, number of mips, and format.
 	 * @param NewSizeX - Width of the texture source data.
@@ -162,6 +178,24 @@ struct FTextureSource
 		const uint8* NewData = NULL
 		);
 
+	/**
+	 * Initialize the source data with the given size, number of mips, and format.
+	 * @param NewSizeX - Width of the texture source data.
+	 * @param NewSizeY - Height of the texture source data.
+	 * @param NewNumSlices - The number of slices in the texture source data.
+	 * @param NewNumMips - The number of mips in the texture source data.
+	 * @param NewFormat - The format in which source data is stored.
+	 * @param NewData - The new source data.
+	 */
+	ENGINE_API void Init(
+		int32 NewSizeX,
+		int32 NewSizeY,
+		int32 NewNumSlices,
+		int32 NewNumMips,
+		ETextureSourceFormat NewFormat,
+		FSharedBuffer NewData
+	);
+
 	/**
 	 * Initializes the source data for a 2D texture with a full mip chain.
 	 * @param NewSizeX - Width of the texture source data.
@@ -425,7 +459,7 @@ private:
 	FSharedBuffer TryDecompressPngData(IImageWrapperModule* ImageWrapperModule) const;
 	/** Attempt to decompress the source data from Jpeg format. All failures will be logged and result in the method returning false */
 	FSharedBuffer TryDecompressJpegData(IImageWrapperModule* ImageWrapperModule) const;
-	
+
 	/** Return true if the source art is not png compressed but could be. */
 	bool CanPNGCompress() const;
 	/** Removes source data. */
@@ -436,6 +470,19 @@ private:
 	int64 CalcBlockSize(int32 BlockIndex) const;
 	int64 CalcLayerSize(int32 BlockIndex, int32 LayerIndex) const;
 
+	void InitLayeredImpl(
+		int32 NewSizeX,
+		int32 NewSizeY,
+		int32 NewNumSlices,
+		int32 NewNumLayers,
+		int32 NewNumMips,
+		const ETextureSourceFormat* NewLayerFormat);
+
+	void InitBlockedImpl(const ETextureSourceFormat* InLayerFormats,
+		const FTextureSourceBlock* InBlocks,
+		int32 InNumLayers,
+		int32 InNumBlocks);
+
 public:
 	/** Uses a hash as the GUID, useful to prevent creating new GUIDs on load for legacy assets. */
 	ENGINE_API void UseHashAsGuid();
diff --git a/Engine/Source/Runtime/Engine/Private/Texture.cpp b/Engine/Source/Runtime/Engine/Private/Texture.cpp
--- a/Engine/Source/Runtime/Engine/Private/Texture.cpp
+++ b/Engine/Source/Runtime/Engine/Private/Texture.cpp
@@ -1068,32 +1068,7 @@
 	int32 InNumBlocks,
 	const uint8** InDataPerBlock)
 {
-	check(InNumBlocks > 0);
-	check(InNumLayers > 0);
+	InitBlockedImpl(InLayerFormats, InBlocks, InNumLayers, InNumBlocks);
 
-	RemoveSourceData();
-
-	BaseBlockX = InBlocks[0].BlockX;
-	BaseBlockY = InBlocks[0].BlockY;
-	SizeX = InBlocks[0].SizeX;
-	SizeY = InBlocks[0].SizeY;
-	NumSlices = InBlocks[0].NumSlices;
-	NumMips = InBlocks[0].NumMips;
-
-	NumLayers = InNumLayers;
-	Format = InLayerFormats[0];
-
-	Blocks.Reserve(InNumBlocks - 1);
-	for (int32 BlockIndex = 1; BlockIndex < InNumBlocks; ++BlockIndex)
-	{
-		Blocks.Add(InBlocks[BlockIndex]);
-	}
-
-	LayerFormat.SetNum(InNumLayers, true);
-	for (int i = 0; i < InNumLayers; ++i)
-	{
-		LayerFormat[i] = InLayerFormats[i];
-	}
-
 	int64 TotalBytes = 0;
 	for (int i = 0; i < InNumBlocks; ++i)
@@ -1101,8 +1076,6 @@
 		TotalBytes += CalcBlockSize(i);
 	}
 
-	checkf(LockState == ELockState::None, TEXT("InitBlocked shouldn't be called in-between LockMip/UnlockMip"));
-
 	FUniqueBuffer Buffer = FUniqueBuffer::Alloc(TotalBytes);
 	uint8* DataPtr = (uint8*)Buffer.GetData();
 
@@ -1122,6 +1095,17 @@
 	BulkData.UpdatePayload(Buffer.MoveToShared());
 }
 
+void FTextureSource::InitBlocked(const ETextureSourceFormat* InLayerFormats,
+	const FTextureSourceBlock* InBlocks,
+	int32 InNumLayers,
+	int32 InNumBlocks,
+	FSharedBuffer NewData)
+{
+	InitBlockedImpl(InLayerFormats, InBlocks, InNumLayers, InNumBlocks);
+
+	BulkData.UpdatePayload(MoveTemp(NewData));
+}
+
 void FTextureSource::InitLayered(
 	int32 NewSizeX,
 	int32 NewSizeY,
@@ -1131,18 +1115,14 @@
 	const ETextureSourceFormat* NewLayerFormat,
 	const uint8* NewData)
 {
-	RemoveSourceData();
-	SizeX = NewSizeX;
-	SizeY = NewSizeY;
-	NumLayers = NewNumLayers;
-	NumSlices = NewNumSlices;
-	NumMips = NewNumMips;
-	Format = NewLayerFormat[0];
-	LayerFormat.SetNum(NewNumLayers, true);
-	for (int i = 0; i < NewNumLayers; ++i)
-	{
-		LayerFormat[i] = NewLayerFormat[i];
-	}
+	InitLayeredImpl(
+		NewSizeX,
+		NewSizeY,
+		NewNumSlices,
+		NewNumLayers,
+		NewNumMips,
+		NewLayerFormat
+		);
 
 	int64 TotalBytes = 0;
 	for (int i = 0; i < NewNumLayers; ++i)
@@ -1150,8 +1130,6 @@
 		TotalBytes += CalcLayerSize(0, i);
 	}
 
-	checkf(LockState == ELockState::None, TEXT("InitLayered shouldn't be called in-between LockMip/UnlockMip"));
-
 	// TODO: Allocating an empty buffer if there is no data to copy from seems like an odd choice but the 
 	// code logic has been doing this for almost a decade so I don't want to change it until I am sure that
 	// it serves no purpose. Given a choice I'd assert on NewData == nullptr instead.
@@ -1165,6 +1143,27 @@
 	}
 }
 
+void FTextureSource::InitLayered(
+	int32 NewSizeX,
+	int32 NewSizeY,
+	int32 NewNumSlices,
+	int32 NewNumLayers,
+	int32 NewNumMips,
+	const ETextureSourceFormat* NewLayerFormat,
+	FSharedBuffer NewData)
+{
+	InitLayeredImpl(
+		NewSizeX,
+		NewSizeY,
+		NewNumSlices,
+		NewNumLayers,
+		NewNumMips,
+		NewLayerFormat
+	);
+
+	BulkData.UpdatePayload(MoveTemp(NewData));
+}
+
 void FTextureSource::Init(
 		int32 NewSizeX,
 		int32 NewSizeY,
@@ -1177,6 +1176,18 @@
 	InitLayered(NewSizeX, NewSizeY, NewNumSlices, 1, NewNumMips, &NewFormat, NewData);
 }
 
+void FTextureSource::Init(
+	int32 NewSizeX,
+	int32 NewSizeY,
+	int32 NewNumSlices,
+	int32 NewNumMips,
+	ETextureSourceFormat NewFormat,
+	FSharedBuffer NewData
+)
+{
+	InitLayered(NewSizeX, NewSizeY, NewNumSlices, 1, NewNumMips, &NewFormat, NewData);
+}
+
 void FTextureSource::Init2DWithMipChain(
 	int32 NewSizeX,
 	int32 NewSizeY,
@@ -2337,4 +2348,63 @@
 	ReadOnlyReference = FSharedBuffer::MakeView(ReadWriteBuffer.Get(), DataLength);
 }
 
+void FTextureSource::InitLayeredImpl(
+	int32 NewSizeX,
+	int32 NewSizeY,
+	int32 NewNumSlices,
+	int32 NewNumLayers,
+	int32 NewNumMips,
+	const ETextureSourceFormat* NewLayerFormat)
+{
+	RemoveSourceData();
+	SizeX = NewSizeX;
+	SizeY = NewSizeY;
+	NumLayers = NewNumLayers;
+	NumSlices = NewNumSlices;
+	NumMips = NewNumMips;
+	Format = NewLayerFormat[0];
+	LayerFormat.SetNum(NewNumLayers, true);
+	for (int i = 0; i < NewNumLayers; ++i)
+	{
+		LayerFormat[i] = NewLayerFormat[i];
+	}
+
+	checkf(LockState == ELockState::None, TEXT("InitLayered shouldn't be called in-between LockMip/UnlockMip"));
+}
+
+void FTextureSource::InitBlockedImpl(const ETextureSourceFormat* InLayerFormats,
+	const FTextureSourceBlock* InBlocks,
+	int32 InNumLayers,
+	int32 InNumBlocks)
+{
+	check(InNumBlocks > 0);
+	check(InNumLayers > 0);
+
+	RemoveSourceData();
+
+	BaseBlockX = InBlocks[0].BlockX;
+	BaseBlockY = InBlocks[0].BlockY;
+	SizeX = InBlocks[0].SizeX;
+	SizeY = InBlocks[0].SizeY;
+	NumSlices = InBlocks[0].NumSlices;
+	NumMips = InBlocks[0].NumMips;
+
+	NumLayers = InNumLayers;
+	Format = InLayerFormats[0];
+
+	Blocks.Reserve(InNumBlocks - 1);
+	for (int32 BlockIndex = 1; BlockIndex < InNumBlocks; ++BlockIndex)
+	{
+		Blocks.Add(InBlocks[BlockIndex]);
+	}
+
+	LayerFormat.SetNum(InNumLayers, true);
+	for (int i = 0; i < InNumLayers; ++i)
+	{
+		LayerFormat[i] = InLayerFormats[i];
+	}
+
+	checkf(LockState == ELockState::None, TEXT("InitBlocked shouldn't be called in-between LockMip/UnlockMip"));
+}
+
 #endif //WITH_EDITOR
diff --git a/Engine/Source/Runtime/ImageWrapper/Public/Formats/HdrImageWrapper.h b/Engine/Source/Runtime/ImageWrapper/Public/Formats/HdrImageWrapper.h
--- a/Engine/Source/Runtime/ImageWrapper/Public/Formats/HdrImageWrapper.h
+++ b/Engine/Source/Runtime/ImageWrapper/Public/Formats/HdrImageWrapper.h
@@ -36,6 +36,7 @@ public:
 
 	void FreeCompressedData();
 
+	using IImageWrapper::GetRaw;
 private:
 	bool GetHeaderLine(const uint8*& BufferPos, char Line[256]);
 
diff --git a/Engine/Source/Runtime/ImageWrapper/Public/IImageWrapper.h b/Engine/Source/Runtime/ImageWrapper/Public/IImageWrapper.h
--- a/Engine/Source/Runtime/ImageWrapper/Public/IImageWrapper.h
+++ b/Engine/Source/Runtime/ImageWrapper/Public/IImageWrapper.h
@@ -2,8 +2,9 @@
 
 #pragma once
 
-#include "CoreTypes.h"
 #include "Containers/Array.h"
+#include "Containers/ArrayView.h"
+#include "CoreTypes.h"
 #include "Templates/SharedPointer.h"
 
 
@@ -128,6 +129,31 @@ public:
 	 */
 	virtual TArray64<uint8> GetCompressed(int32 Quality = 0) = 0;
 
+	/**
+	 * Get the raw version of the image and write to the array view
+	 * (Note: It may consume the data set in the SetRaw function if it was set before)
+	 *
+	 * @param InFormat How we want to manipulate the RGB data.
+	 * @param InBitDepth The output bit-depth per channel, normally 8.
+	 * @param OutRawData Will contain the uncompressed raw data.
+	 * @return true on success, false otherwise.
+	 */
+	virtual bool GetRaw(const ERGBFormat InFormat, int32 InBitDepth, TArrayView64<uint8> OutRawData)
+	{
+		// Todo We should make this function virtual pure and implement this properly in each ImageWrapper
+		TArray64<uint8> TmpRawData;
+		if (GetRaw(InFormat, InBitDepth, TmpRawData))
+		{
+			if (ensureMsgf(TmpRawData.Num() == OutRawData.Num(), TEXT("The view doesn't have the proper size to receive the texture.")))
+			{
+				FPlatformMemory::Memcpy(OutRawData.GetData(), TmpRawData.GetData(), OutRawData.Num());
+				return true;
+			}
+		}
+
+		return false;
+	}
+
 	/**  
 	 * Gets the raw data.
 	 * (Note: It may consume the data set in the SetRaw function if it was set before)
diff --git a/Engine/Source/Runtime/Interchange/Core/Public/InterchangeFactoryBase.h b/Engine/Source/Runtime/Interchange/Core/Public/InterchangeFactoryBase.h
--- a/Engine/Source/Runtime/Interchange/Core/Public/InterchangeFactoryBase.h
+++ b/Engine/Source/Runtime/Interchange/Core/Public/InterchangeFactoryBase.h
@@ -82,7 +82,7 @@ public:
 	 *
 	 * @Note override function should verify the AssetNode in Arguments match the expected type for the factory.
 	 */
-	virtual UObject* CreateEmptyAsset(const FCreateAssetParams& Arguments) const
+	virtual UObject* CreateEmptyAsset(const FCreateAssetParams& Arguments)
 	{
 		return nullptr;
 	}
@@ -95,7 +95,7 @@ public:
 	 *
 	 * @Note override function should verify the AssetNode in Arguments match the expected type for the factory.
 	 */
-	virtual UObject* CreateAsset(const FCreateAssetParams& Arguments) const
+	virtual UObject* CreateAsset(const FCreateAssetParams& Arguments)
 	{
 		return nullptr;
 	}
@@ -113,12 +113,18 @@ public:
 	{
 		/** The source data, mainly use to set the asset import data file. TODO: we have to refactor UAssetImportData, the source data should be the base class for this now */
 		const UInterchangeSourceData* SourceData = nullptr;
+		UInterchangeBaseNode* FactoryNode = nullptr;
+
 
 		/** The UObject  we want to execute code on*/
 		UObject* ImportedObject = nullptr;
 		FString NodeUniqueID;
 		UInterchangeBaseNodeContainer* NodeContainer = nullptr;
 		TArray<UInterchangePipelineBase*> Pipelines;
+
+
+		EReimportStrategyFlags ReimportStrategyFlags;
+		bool bIsReimport  = false;
 ;
 	};
 
@@ -126,7 +132,7 @@ public:
 	 * This function is call in the pre completion task on the main thread, use it to call main thread post creation step for your assets
 	 * @note - This function is called when starting the pre completion task (before PostEditChange is called for the asset).
 	 */
-	virtual void PreImportPreCompletedCallback(const FImportPreCompletedCallbackParams& Arguments) const
+	virtual void PreImportPreCompletedCallback(const FImportPreCompletedCallbackParams& Arguments)
 	{
 		check(IsInGameThread());
 		return;
@@ -136,7 +142,7 @@ public:
 	 * This function is call in the pre completion task on the main thread, use it to call main thread post creation step for your assets
 	 * @note - This function is called at the end of the pre completion task (after PostEditChange is called for the asset).
 	 */
-	virtual void PostImportPreCompletedCallback(const FImportPreCompletedCallbackParams& Arguments) const
+	virtual void PostImportPreCompletedCallback(const FImportPreCompletedCallbackParams& Arguments)
 	{
 		check(IsInGameThread());
 		return;
@@ -146,7 +152,7 @@ public:
 	 * This function is used to add the given message object directly into the results for this operation.
 	 */
 	template <typename T>
-	T* AddMessage() const
+	T* AddMessage()
 	{
 		check(Results != nullptr);
 		T* Item = Results->Add<T>();
@@ -154,7 +160,7 @@ public:
 	}
 
 
-	void AddMessage(UInterchangeResult* Item) const
+	void AddMessage(UInterchangeResult* Item)
 	{
 		check(Results != nullptr);
 		Results->Add(Item);
diff --git a/Engine/Source/Runtime/Interchange/Engine/Private/Tasks/InterchangeTaskCompletion.cpp b/Engine/Source/Runtime/Interchange/Engine/Private/Tasks/InterchangeTaskCompletion.cpp
--- a/Engine/Source/Runtime/Interchange/Engine/Private/Tasks/InterchangeTaskCompletion.cpp
+++ b/Engine/Source/Runtime/Interchange/Engine/Private/Tasks/InterchangeTaskCompletion.cpp
@@ -56,9 +56,12 @@ void UE::Interchange::FTaskPreCompletion::DoTask(ENamedThreads::Type CurrentThre
 				UInterchangeFactoryBase::FImportPreCompletedCallbackParams Arguments;
 				Arguments.ImportedObject = Asset;
 				Arguments.SourceData = AsyncHelper->SourceDatas[SourceIndex];
-				Arguments.NodeUniqueID = AssetInfo.NodeUniqueId;
+				Arguments.FactoryNode = AssetInfo.FactoryNode;
+				// Should we assert if there is no factory node?
+				Arguments.NodeUniqueID = AssetInfo.FactoryNode ? AssetInfo.FactoryNode->GetUniqueID() : FString();
 				Arguments.NodeContainer = AsyncHelper->BaseNodeContainers[SourceIndex].Get();
 				Arguments.Pipelines = AsyncHelper->Pipelines;
+				Arguments.bIsReimport = AssetInfo.bIsReimport;
 				AssetInfo.Factory->PreImportPreCompletedCallback(Arguments);
 			}
 
@@ -95,9 +98,11 @@ void UE::Interchange::FTaskPreCompletion::DoTask(ENamedThreads::Type CurrentThre
 				UInterchangeFactoryBase::FImportPreCompletedCallbackParams Arguments;
 				Arguments.ImportedObject = Asset;
 				Arguments.SourceData = AsyncHelper->SourceDatas[SourceIndex];
-				Arguments.NodeUniqueID = AssetInfo.NodeUniqueId;
+				Arguments.FactoryNode = AssetInfo.FactoryNode;
+				Arguments.NodeUniqueID = AssetInfo.FactoryNode ? AssetInfo.FactoryNode->GetUniqueID() : FString();
 				Arguments.NodeContainer = AsyncHelper->BaseNodeContainers[SourceIndex].Get();
 				Arguments.Pipelines = AsyncHelper->Pipelines;
+				Arguments.bIsReimport = AssetInfo.bIsReimport;
 				AssetInfo.Factory->PostImportPreCompletedCallback(Arguments);
 			}
 		}
diff --git a/Engine/Source/Runtime/Interchange/Engine/Private/Tasks/InterchangeTaskCreateAsset.cpp b/Engine/Source/Runtime/Interchange/Engine/Private/Tasks/InterchangeTaskCreateAsset.cpp
--- a/Engine/Source/Runtime/Interchange/Engine/Private/Tasks/InterchangeTaskCreateAsset.cpp
+++ b/Engine/Source/Runtime/Interchange/Engine/Private/Tasks/InterchangeTaskCreateAsset.cpp
@@ -153,7 +153,8 @@ void UE::Interchange::FTaskCreatePackage::DoTask(ENamedThreads::Type CurrentThre
 			UE::Interchange::FImportAsyncHelper::FImportedAssetInfo& AssetInfo = ImportedInfos.AddDefaulted_GetRef();
 			AssetInfo.ImportAsset = NodeAsset;
 			AssetInfo.Factory = Factory;
-			AssetInfo.NodeUniqueId = Node->GetUniqueID();
+			AssetInfo.FactoryNode = Node;
+			AssetInfo.bIsReimport = bool(AsyncHelper->TaskData.ReimportObject);
 			Node->ReferenceObject = FSoftObjectPath(NodeAsset);
 		}
 	}
@@ -284,7 +285,8 @@ void UE::Interchange::FTaskCreateAsset::DoTask(ENamedThreads::Type CurrentThread
 				UE::Interchange::FImportAsyncHelper::FImportedAssetInfo& AssetInfo = ImportedInfos.AddDefaulted_GetRef();
 				AssetInfo.ImportAsset = NodeAsset;
 				AssetInfo.Factory = Factory;
-				AssetInfo.NodeUniqueId = Node->GetUniqueID();
+				AssetInfo.FactoryNode = Node;
+				AssetInfo.bIsReimport = bool(AsyncHelper->TaskData.ReimportObject);
 			}
 
 			// Fill in destination asset and type in any results which have been added previously by a translator or pipeline, now that we have a corresponding factory.
diff --git a/Engine/Source/Runtime/Interchange/Engine/Private/Tasks/InterchangeTaskPipeline.cpp b/Engine/Source/Runtime/Interchange/Engine/Private/Tasks/InterchangeTaskPipeline.cpp
--- a/Engine/Source/Runtime/Interchange/Engine/Private/Tasks/InterchangeTaskPipeline.cpp
+++ b/Engine/Source/Runtime/Interchange/Engine/Private/Tasks/InterchangeTaskPipeline.cpp
@@ -69,7 +69,7 @@ void UE::Interchange::FTaskPipelinePostImport::DoTask(ENamedThreads::Type Curren
 			ImportAssets.Reserve(ImportedInfos.Num());
 			for (UE::Interchange::FImportAsyncHelper::FImportedAssetInfo& ImportedInfo : ImportedInfos)
 			{
-				NodeUniqueIDs.Add(ImportedInfo.NodeUniqueId);
+				NodeUniqueIDs.Add(ImportedInfo.FactoryNode->GetUniqueID());
 				ImportAssets.Add(ImportedInfo.ImportAsset);
 			}
 		}
diff --git a/Engine/Source/Runtime/Interchange/Engine/Public/InterchangeManager.h b/Engine/Source/Runtime/Interchange/Engine/Public/InterchangeManager.h
--- a/Engine/Source/Runtime/Interchange/Engine/Public/InterchangeManager.h
+++ b/Engine/Source/Runtime/Interchange/Engine/Public/InterchangeManager.h
@@ -177,7 +177,8 @@ namespace UE
 			{
 				UObject* ImportAsset;
 				UInterchangeFactoryBase* Factory;
-				FString NodeUniqueId;
+				UInterchangeBaseNode* FactoryNode;
+				bool bIsReimport;
 			};
 
 			FCriticalSection ImportedAssetsPerSourceIndexLock;
