Allows VirtualizedBulkData to take a shared buffer with a precomtued PayloadId
Changed some init function of the stuct  FTextureSource to take the new struct FSharedBufferWithID.
Changed the interchange texture factory to precompute the PayloadId outside of the game thread.
On a very large udim texture (7.6 GB of source data), this can help us avoid a 2 seconds stall on the gamethread.

diff --git a/Engine/Plugins/Runtime/Interchange/Source/Import/Private/Texture/InterchangeTextureFactory.cpp b/Engine/Plugins/Runtime/Interchange/Source/Import/Private/Texture/InterchangeTextureFactory.cpp
--- a/Engine/Plugins/Runtime/Interchange/Source/Import/Private/Texture/InterchangeTextureFactory.cpp
+++ b/Engine/Plugins/Runtime/Interchange/Source/Import/Private/Texture/InterchangeTextureFactory.cpp
@@ -263,7 +263,7 @@ namespace UE::Interchange::Private::InterchangeTextureFactory
 	}
 
 #if WITH_EDITORONLY_DATA
-	void SetupTextureSourceData(UTexture* Texture, FImportImage& Image)
+	void SetupTextureSourceData(UTexture* Texture, const FImportImage& Image, UE::Virtualization::FVirtualizedUntypedBulkData::FSharedBufferWithID&& BufferAndId)
 	{
 		Texture->Source.Init(
 			Image.SizeX,
@@ -271,7 +271,7 @@ namespace UE::Interchange::Private::InterchangeTextureFactory
 			/*NumSlices=*/ 1,
 			Image.NumMips,
 			Image.Format,
-			Image.RawData.MoveToShared()
+			MoveTemp(BufferAndId)
 		);
 
 		Texture->CompressionSettings = Image.CompressionSettings;
@@ -285,58 +285,47 @@ namespace UE::Interchange::Private::InterchangeTextureFactory
 		}
 	}
 
-	bool SetupTexture2DSourceData(UTexture2D* Texture2D, FImportBlockedImage& BlockedImage)
+	void SetupTexture2DSourceData(UTexture2D* Texture2D, const FImportBlockedImage& BlockedImage, UE::Virtualization::FVirtualizedUntypedBulkData::FSharedBufferWithID&& BufferAndId)
 	{
-		if (BlockedImage.IsValid())
+		if (BlockedImage.BlocksData.Num() > 1)
 		{
-			if (BlockedImage.BlocksData.Num() > 1)
-			{
-
-				Texture2D->Source.InitBlocked(
-					&BlockedImage.Format,
-					BlockedImage.BlocksData.GetData(),
-					/*InNumLayers=*/ 1,
-					BlockedImage.BlocksData.Num(),
-					BlockedImage.RawData.MoveToShared()
-					);
-
-				Texture2D->CompressionSettings = BlockedImage.CompressionSettings;
-				Texture2D->SRGB = BlockedImage.bSRGB;
-				Texture2D->VirtualTextureStreaming = true;
+			Texture2D->Source.InitBlocked(
+				&BlockedImage.Format,
+				BlockedImage.BlocksData.GetData(),
+				/*InNumLayers=*/ 1,
+				BlockedImage.BlocksData.Num(),
+				MoveTemp(BufferAndId)
+				);
 
-				if (BlockedImage.MipGenSettings.IsSet())
-				{
-					// if the source has mips we keep the mips by default, unless the user changes that
-					Texture2D->MipGenSettings = BlockedImage.MipGenSettings.GetValue();
-				}
+			Texture2D->CompressionSettings = BlockedImage.CompressionSettings;
+			Texture2D->SRGB = BlockedImage.bSRGB;
+			Texture2D->VirtualTextureStreaming = true;
 
-				return true;
-			}
-			else
+			if (BlockedImage.MipGenSettings.IsSet())
 			{
-				//Import as a normal texture
-				FImportImage Image;
-				Image.Format = BlockedImage.Format;
-				Image.CompressionSettings = BlockedImage.CompressionSettings;
-				Image.bSRGB = BlockedImage.bSRGB;
-				Image.MipGenSettings = BlockedImage.MipGenSettings;
-
-				const FTextureSourceBlock& Block = BlockedImage.BlocksData[0];
-				Image.SizeX = Block.SizeX;
-				Image.SizeY = Block.SizeY;
-				Image.NumMips = Block.NumMips;
-
-				Image.RawData = MoveTemp(BlockedImage.RawData);
-
-				SetupTextureSourceData(Texture2D, Image);
-				return true;
+				// if the source has mips we keep the mips by default, unless the user changes that
+				Texture2D->MipGenSettings = BlockedImage.MipGenSettings.GetValue();
 			}
 		}
+		else
+		{
+			//Import as a normal texture
+			FImportImage Image;
+			Image.Format = BlockedImage.Format;
+			Image.CompressionSettings = BlockedImage.CompressionSettings;
+			Image.bSRGB = BlockedImage.bSRGB;
+			Image.MipGenSettings = BlockedImage.MipGenSettings;
 
-		return false;
+			const FTextureSourceBlock& Block = BlockedImage.BlocksData[0];
+			Image.SizeX = Block.SizeX;
+			Image.SizeY = Block.SizeY;
+			Image.NumMips = Block.NumMips;
+
+			SetupTextureSourceData(Texture2D, Image, MoveTemp(BufferAndId));
+		}
 	}
 
-	void SetupTextureSourceData(UTexture* Texture, FImportSlicedImage& SlicedImage)
+	void SetupTextureSourceData(UTexture* Texture, const FImportSlicedImage& SlicedImage, UE::Virtualization::FVirtualizedUntypedBulkData::FSharedBufferWithID&& BufferAndId)
 	{
 		Texture->Source.InitLayered(
 			SlicedImage.SizeX,
@@ -345,7 +334,7 @@ namespace UE::Interchange::Private::InterchangeTextureFactory
 			1,
 			SlicedImage.NumMips,
 			&SlicedImage.Format,
-			SlicedImage.RawData.MoveToShared()
+			MoveTemp(BufferAndId)
 			);
 
 		Texture->CompressionSettings = SlicedImage.CompressionSettings;
@@ -358,10 +347,10 @@ namespace UE::Interchange::Private::InterchangeTextureFactory
 		}
 	}
 
-	void SetupTextureSourceData(UTextureLightProfile* TextureLightProfile, FImportLightProfile& LightProfile)
+	void SetupTextureSourceData(UTextureLightProfile* TextureLightProfile, const FImportLightProfile& LightProfile, UE::Virtualization::FVirtualizedUntypedBulkData::FSharedBufferWithID&& BufferAndId)
 	{
-		FImportImage& ImportImage = LightProfile;
-		SetupTextureSourceData(TextureLightProfile, ImportImage);
+		const FImportImage& ImportImage = LightProfile;
+		SetupTextureSourceData(TextureLightProfile, ImportImage, MoveTemp(BufferAndId));
 
 		TextureLightProfile->Brightness = LightProfile.Brightness;
 		TextureLightProfile->TextureMultiplier = LightProfile.TextureMultiplier;
@@ -394,24 +383,23 @@ namespace UE::Interchange::Private::InterchangeTextureFactory
 		return false;
 	}
 
-	bool SetupTexture2DSourceData(UTexture2D* Texture2D, FTexturePayloadVariant& TexturePayload)
+	void SetupTexture2DSourceData(UTexture2D* Texture2D, FProcessedPayload& ProcessedPayload)
 	{
-		if (TOptional<FImportBlockedImage>* BlockedImage = TexturePayload.TryGet<TOptional<FImportBlockedImage>>())
+		if (TOptional<FImportBlockedImage>* BlockedImage = ProcessedPayload.SettingsFromPayload.TryGet<TOptional<FImportBlockedImage>>())
 		{
 			if (BlockedImage->IsSet())
 			{
-				return SetupTexture2DSourceData(Texture2D, BlockedImage->GetValue());
+				SetupTexture2DSourceData(Texture2D, BlockedImage->GetValue(), MoveTemp(ProcessedPayload.PayloadAndId));
 			}
 		}
-		else if (TOptional<FImportImage>* Image = TexturePayload.TryGet<TOptional<FImportImage>>())
+		else if (TOptional<FImportImage>* Image = ProcessedPayload.SettingsFromPayload.TryGet<TOptional<FImportImage>>())
 		{
 			if (Image->IsSet())
 			{
-				SetupTextureSourceData(Texture2D, Image->GetValue());
-				return true;
+				SetupTextureSourceData(Texture2D, Image->GetValue(), MoveTemp(ProcessedPayload.PayloadAndId));
 			}
 		}
-		else if (TOptional<FImportLightProfile>* OptionalLightProfile = TexturePayload.TryGet<TOptional<FImportLightProfile>>())
+		else if (TOptional<FImportLightProfile>* OptionalLightProfile = ProcessedPayload.SettingsFromPayload.TryGet<TOptional<FImportLightProfile>>())
 		{
 			if (OptionalLightProfile->IsSet())
 			{
@@ -419,18 +407,17 @@ namespace UE::Interchange::Private::InterchangeTextureFactory
 
 				if (UTextureLightProfile* TextureLightProfile = Cast<UTextureLightProfile>(Texture2D))
 				{
-					SetupTextureSourceData(TextureLightProfile, LightProfile);
+					SetupTextureSourceData(TextureLightProfile, LightProfile, MoveTemp(ProcessedPayload.PayloadAndId));
 				}
 				else
 				{
-					SetupTextureSourceData(Texture2D, LightProfile);
+					SetupTextureSourceData(Texture2D, LightProfile, MoveTemp(ProcessedPayload.PayloadAndId));
 				}
 			}
-
-			return true;
 		}
 
-		return false;
+		// The TexturePayload should be validated before calling this function
+		checkNoEntry();
 	}
 
 	bool CanSetupTextureCubeSourceData(FTexturePayloadVariant& TexturePayload)
@@ -460,9 +447,9 @@ namespace UE::Interchange::Private::InterchangeTextureFactory
 		return false;
 	}
 
-	bool SetupTextureCubeSourceData(UTextureCube* TextureCube, FTexturePayloadVariant& TexturePayload)
+	void SetupTextureCubeSourceData(UTextureCube* TextureCube, FProcessedPayload& ProcessedPayload)
 	{
-		if (TOptional<FImportSlicedImage>* OptionalSlicedImage = TexturePayload.TryGet<TOptional<FImportSlicedImage>>())
+		if (TOptional<FImportSlicedImage>* OptionalSlicedImage = ProcessedPayload.SettingsFromPayload.TryGet<TOptional<FImportSlicedImage>>())
 		{
 			if (OptionalSlicedImage->IsSet())
 			{
@@ -471,31 +458,29 @@ namespace UE::Interchange::Private::InterchangeTextureFactory
 				// Cube texture always have six slice
 				if (SlicedImage.NumSlice == 6)
 				{
-					SetupTextureSourceData(TextureCube, SlicedImage);
-					return true;
+					SetupTextureSourceData(TextureCube, SlicedImage, MoveTemp(ProcessedPayload.PayloadAndId));
 				}
 			}
 		}
-		else if (TOptional<FImportImage>* OptionalImage = TexturePayload.TryGet<TOptional<FImportImage>>())
+		else if (TOptional<FImportImage>* OptionalImage = ProcessedPayload.SettingsFromPayload.TryGet<TOptional<FImportImage>>())
 		{
 			if (OptionalImage->IsSet())
 			{
 				FImportImage& Image = OptionalImage->GetValue();
-				SetupTextureSourceData(TextureCube, Image);
-				return true;
+				SetupTextureSourceData(TextureCube, Image, MoveTemp(ProcessedPayload.PayloadAndId));
 			}
 		}
-		else if (TOptional<FImportLightProfile>* OptionalLightProfile = TexturePayload.TryGet<TOptional<FImportLightProfile>>())
+		else if (TOptional<FImportLightProfile>* OptionalLightProfile = ProcessedPayload.SettingsFromPayload.TryGet<TOptional<FImportLightProfile>>())
 		{
 			if (OptionalLightProfile->IsSet())
 			{
 				FImportLightProfile& LightProfile = OptionalLightProfile->GetValue();
-				SetupTextureSourceData(TextureCube, LightProfile);
-				return true;
+				SetupTextureSourceData(TextureCube, LightProfile, MoveTemp(ProcessedPayload.PayloadAndId));
 			}
 		}
 
-		return false;
+		// The TexturePayload should be validated before calling this function
+		checkNoEntry();
 	}
 
 	bool CanSetupTexture2DArraySourceData(FTexturePayloadVariant& TexturePayload)
@@ -525,38 +510,36 @@ namespace UE::Interchange::Private::InterchangeTextureFactory
 		return false;
 	}
 
-	bool SetupTexture2DArraySourceData(UTexture2DArray* Texture2DArray, FTexturePayloadVariant& TexturePayload)
+	void SetupTexture2DArraySourceData(UTexture2DArray* Texture2DArray, FProcessedPayload& ProcessedPayload)
 	{
-		if (TOptional<FImportSlicedImage>* OptionalSlicedImage = TexturePayload.TryGet<TOptional<FImportSlicedImage>>())
+		if (TOptional<FImportSlicedImage>* OptionalSlicedImage = ProcessedPayload.SettingsFromPayload.TryGet<TOptional<FImportSlicedImage>>())
 		{
 			if (OptionalSlicedImage->IsSet())
 			{
 				FImportSlicedImage& SlicedImage = OptionalSlicedImage->GetValue();
 
-				SetupTextureSourceData(Texture2DArray, SlicedImage);
-				return true;
+				SetupTextureSourceData(Texture2DArray, SlicedImage, MoveTemp(ProcessedPayload.PayloadAndId));
 			}
 		}
-		else if (TOptional<FImportImage>* OptionalImage = TexturePayload.TryGet<TOptional<FImportImage>>())
+		else if (TOptional<FImportImage>* OptionalImage = ProcessedPayload.SettingsFromPayload.TryGet<TOptional<FImportImage>>())
 		{
 			if (OptionalImage->IsSet())
 			{
 				FImportImage& Image = OptionalImage->GetValue();
-				SetupTextureSourceData(Texture2DArray, Image);
-				return true;
+				SetupTextureSourceData(Texture2DArray, Image, MoveTemp(ProcessedPayload.PayloadAndId));
 			}
 		}
-		else if (TOptional<FImportLightProfile>* OptionalLightProfile = TexturePayload.TryGet<TOptional<FImportLightProfile>>())
+		else if (TOptional<FImportLightProfile>* OptionalLightProfile = ProcessedPayload.SettingsFromPayload.TryGet<TOptional<FImportLightProfile>>())
 		{
 			if (OptionalLightProfile->IsSet())
 			{
 				FImportLightProfile& LightProfile = OptionalLightProfile->GetValue();
-				SetupTextureSourceData(Texture2DArray, LightProfile);
+				SetupTextureSourceData(Texture2DArray, LightProfile, MoveTemp(ProcessedPayload.PayloadAndId));
 			}
-			return true;
 		}
 
-		return false;
+		// The TexturePayload should be validated before calling this function
+		checkNoEntry();
 	}
 
 	void LogErrorInvalidPayload(const FString& TextureClass, const FString& ObjectName)
@@ -564,6 +547,38 @@ namespace UE::Interchange::Private::InterchangeTextureFactory
 		UE_LOG(LogInterchangeImport, Error, TEXT("UInterchangeTextureFactory: The Payload was invalid for a %s. (%s)"), *TextureClass, *ObjectName);
 	}
 
+	FSharedBuffer MoveRawDataToSharedBuffer(FTexturePayloadVariant& TexturePayload)
+	{
+		if (TOptional<FImportBlockedImage>* BlockedImage = TexturePayload.TryGet<TOptional<FImportBlockedImage>>())
+		{
+			return (*BlockedImage)->RawData.MoveToShared();
+		}
+		else if (TOptional<FImportImage>* Image = TexturePayload.TryGet<TOptional<FImportImage>>())
+		{
+			return (*Image)->RawData.MoveToShared();
+		}
+		else if (TOptional<FImportSlicedImage>* SlicedImage = TexturePayload.TryGet<TOptional<FImportSlicedImage>>())
+		{
+			return (*SlicedImage)->RawData.MoveToShared();
+		}
+		else if (TOptional<FImportLightProfile>* LightProfile = TexturePayload.TryGet<TOptional<FImportLightProfile>>())
+		{
+			return (*LightProfile)->RawData.MoveToShared();
+		}
+
+		// The TexturePayload should be validated before calling this function
+		checkNoEntry();
+		return {};
+	}
+
+	FProcessedPayload& FProcessedPayload::operator=(UE::Interchange::Private::InterchangeTextureFactory::FTexturePayloadVariant&& InPayloadVariant)
+	{
+		SettingsFromPayload = MoveTemp(InPayloadVariant);
+		PayloadAndId = MoveRawDataToSharedBuffer(SettingsFromPayload);
+
+		return *this;
+	}
+
 #endif // WITH_EDITORONLY_DATA
  }
 
@@ -691,7 +706,7 @@ UObject* UInterchangeTextureFactory::CreateAsset(const FCreateAssetParams& Argum
 		return nullptr;
 	}
 
-	TexturePayload = GetTexturePayload(Arguments.SourceData, PayLoadKey.GetValue(), TextureNodeVariant, Arguments.Translator);
+	FTexturePayloadVariant TexturePayload = GetTexturePayload(Arguments.SourceData, PayLoadKey.GetValue(), TextureNodeVariant, Arguments.Translator);
 
 	if(TexturePayload.IsType<FEmptyVariantState>())
 	{
@@ -755,6 +770,9 @@ UObject* UInterchangeTextureFactory::CreateAsset(const FCreateAssetParams& Argum
 	//Getting the file Hash will cache it into the source data
 	Arguments.SourceData->GetFileContentHash();
 
+	// This will hash the payload to generate a unique ID before passing it to the virtualized bulkdata
+	ProcessedPayload = MoveTemp(TexturePayload);
+
 	//The interchange completion task (call in the GameThread after the factories pass), will call PostEditChange which will trig another asynchronous system that will build all texture in parallel
 
 	return Texture;
@@ -782,27 +800,15 @@ void UInterchangeTextureFactory::PreImportPreCompletedCallback(const FImportPreC
 		// Setup source data
 		if (UTexture2D* Texture2D = Cast<UTexture2D>(Texture))
 		{
-			if (!SetupTexture2DSourceData(Texture2D, TexturePayload))
-			{
-				LogErrorInvalidPayload(Arguments.ImportedObject->GetClass()->GetName(), Arguments.ImportedObject->GetName());
-				return;			
-			}
+			SetupTexture2DSourceData(Texture2D, ProcessedPayload);
 		}
 		else if (UTextureCube* TextureCube = Cast<UTextureCube>(Texture))
 		{
-			if (!SetupTextureCubeSourceData(TextureCube, TexturePayload))
-			{
-				LogErrorInvalidPayload(Arguments.ImportedObject->GetClass()->GetName(), Arguments.ImportedObject->GetName());
-				return;
-			}
+			SetupTextureCubeSourceData(TextureCube, ProcessedPayload);
 		}
 		else if (UTexture2DArray* Texture2DArray = Cast<UTexture2DArray>(Arguments.ImportedObject))
 		{
-			if (!SetupTexture2DArraySourceData(Texture2DArray, TexturePayload))
-			{
-				LogErrorInvalidPayload(Arguments.ImportedObject->GetClass()->GetName(), Arguments.ImportedObject->GetName());
-				return;
-			}
+			SetupTexture2DArraySourceData(Texture2DArray, ProcessedPayload);
 		}
 		else
 		{
@@ -851,4 +857,4 @@ void UInterchangeTextureFactory::PreImportPreCompletedCallback(const FImportPreC
 		Texture->AssetImportData = UE::Interchange::FFactoryCommon::UpdateImportAssetData(UpdateImportAssetDataParameters);
 	}
 #endif
-}
\ No newline at end of file
+}
diff --git a/Engine/Plugins/Runtime/Interchange/Source/Import/Public/Texture/InterchangeTextureFactory.h b/Engine/Plugins/Runtime/Interchange/Source/Import/Public/Texture/InterchangeTextureFactory.h
--- a/Engine/Plugins/Runtime/Interchange/Source/Import/Public/Texture/InterchangeTextureFactory.h
+++ b/Engine/Plugins/Runtime/Interchange/Source/Import/Public/Texture/InterchangeTextureFactory.h
@@ -12,6 +12,7 @@
 #include "Texture/InterchangeTexturePayloadData.h"
 #include "UObject/Object.h"
 #include "UObject/ObjectMacros.h"
+#include "Virtualization/VirtualizedBulkData.h"
 
 #include "InterchangeTextureFactory.generated.h"
 
@@ -22,6 +23,21 @@ namespace UE::Interchange::Private::InterchangeTextureFactory
 		, TOptional<FImportBlockedImage>
 		, TOptional<FImportSlicedImage>
 		, TOptional<FImportLightProfile>>;
+
+	struct FProcessedPayload
+	{
+		FProcessedPayload() = default;
+		FProcessedPayload(FProcessedPayload&&) = default;
+		FProcessedPayload& operator=(FProcessedPayload&&) = default;
+
+		FProcessedPayload(const FProcessedPayload&) = delete;
+		FProcessedPayload& operator=(const FProcessedPayload&) = delete;
+
+		FProcessedPayload& operator=(UE::Interchange::Private::InterchangeTextureFactory::FTexturePayloadVariant&& InPayloadVariant);
+
+		UE::Interchange::Private::InterchangeTextureFactory::FTexturePayloadVariant SettingsFromPayload;
+		UE::Virtualization::FVirtualizedUntypedBulkData::FSharedBufferWithID PayloadAndId;
+	};
 }
 
 UCLASS(BlueprintType)
@@ -41,7 +57,7 @@ public:
 	//////////////////////////////////////////////////////////////////////////
 
 private:
-	UE::Interchange::Private::InterchangeTextureFactory::FTexturePayloadVariant TexturePayload;
+	UE::Interchange::Private::InterchangeTextureFactory::FProcessedPayload ProcessedPayload;
 };
 
 
diff --git a/Engine/Source/Runtime/CoreUObject/Private/Virtualization/VirtualizedBulkData.cpp b/Engine/Source/Runtime/CoreUObject/Private/Virtualization/VirtualizedBulkData.cpp
--- a/Engine/Source/Runtime/CoreUObject/Private/Virtualization/VirtualizedBulkData.cpp
+++ b/Engine/Source/Runtime/CoreUObject/Private/Virtualization/VirtualizedBulkData.cpp
@@ -941,6 +941,43 @@ FGuid FVirtualizedUntypedBulkData::GetIdentifier() const
 	return BulkDataId;
 }
 
+void FVirtualizedUntypedBulkData::UpdatePayloadImpl(FSharedBuffer&& InPayload, FPayloadId&& InPayloadID, FName InCompressionFormat)
+{
+	TRACE_CPUPROFILER_EVENT_SCOPE(FVirtualizedUntypedBulkData::UpdatePayloadImpl);
+
+	UnloadData();
+
+	// Make sure that we own the memory in the shared buffer
+	Payload = InPayload.MakeOwned();
+	PayloadSize = (int64)Payload.GetSize();
+	PayloadContentId = MoveTemp(InPayloadID);
+
+	EnumRemoveFlags(Flags, EFlags::IsVirtualized |
+		EFlags::DisablePayloadCompression |
+		EFlags::ReferencesLegacyFile |
+		EFlags::LegacyFileIsCompressed |
+		EFlags::LegacyKeyWasGuidDerived);
+
+	SetCompressionFormat(InCompressionFormat);
+
+	PackagePath.Empty();
+	PackageSegment = EPackageSegment::Header;
+	OffsetInFile = INDEX_NONE;
+
+	if (PayloadSize > 0)
+	{
+		if (!BulkDataId.IsValid())
+		{
+			BulkDataId = FGuid::NewGuid();
+		}
+		Register(nullptr);
+	}
+	else
+	{
+		Unregister();
+	}
+}
+
 FCompressedBuffer FVirtualizedUntypedBulkData::GetDataInternal() const
 {
 	TRACE_CPUPROFILER_EVENT_SCOPE(FVirtualizedUntypedBulkData::GetDataInternal);
@@ -1006,38 +1043,19 @@ TFuture<FCompressedBuffer>FVirtualizedUntypedBulkData::GetCompressedPayload() co
 void FVirtualizedUntypedBulkData::UpdatePayload(FSharedBuffer InPayload, FName InCompressionFormat)
 {
 	TRACE_CPUPROFILER_EVENT_SCOPE(FVirtualizedUntypedBulkData::UpdatePayload);
+	FPayloadId NewPayloadId(InPayload);
+	UpdatePayloadImpl(MoveTemp(InPayload), MoveTemp(NewPayloadId), InCompressionFormat);
+}
 
-	UnloadData();
-
-	// Make sure that we own the memory in the shared buffer
-	Payload = InPayload.MakeOwned();
-	PayloadSize = (int64)Payload.GetSize();
-	PayloadContentId = FPayloadId(Payload);
-
-	EnumRemoveFlags(Flags,	EFlags::IsVirtualized | 
-							EFlags::DisablePayloadCompression |
-							EFlags::ReferencesLegacyFile |
-							EFlags::LegacyFileIsCompressed |
-							EFlags::LegacyKeyWasGuidDerived);
-
-	SetCompressionFormat(InCompressionFormat);
-	
-	PackagePath.Empty();
-	PackageSegment = EPackageSegment::Header;
-	OffsetInFile = INDEX_NONE;
+FVirtualizedUntypedBulkData::FSharedBufferWithID::FSharedBufferWithID(FSharedBuffer InPayload)
+	: Payload(MoveTemp(InPayload))
+	, PayloadId(Payload)
+{
+}
 
-	if (PayloadSize > 0)
-	{
-		if (!BulkDataId.IsValid())
-		{
-			BulkDataId = FGuid::NewGuid();
-		}
-		Register(nullptr);
-	}
-	else
-	{
-		Unregister();
-	}
+void FVirtualizedUntypedBulkData::UpdatePayload(FSharedBufferWithID InPayload, FName InCompressionFormat)
+{
+	UpdatePayloadImpl(MoveTemp(InPayload.Payload), MoveTemp(InPayload.PayloadId), InCompressionFormat);
 }
 
 void FVirtualizedUntypedBulkData::SetCompressionFormat(FName InCompressionFormat)
diff --git a/Engine/Source/Runtime/CoreUObject/Public/Virtualization/VirtualizedBulkData.h b/Engine/Source/Runtime/CoreUObject/Public/Virtualization/VirtualizedBulkData.h
--- a/Engine/Source/Runtime/CoreUObject/Public/Virtualization/VirtualizedBulkData.h
+++ b/Engine/Source/Runtime/CoreUObject/Public/Virtualization/VirtualizedBulkData.h
@@ -140,6 +140,50 @@ public:
 	 */
 	void UpdatePayload(FSharedBuffer InPayload, FName CompressionFormat = NAME_Default);
 
+	/**
+	 * Utility struct used to compute the Payload ID before calling UpdatePayload
+	 */
+	struct COREUOBJECT_API FSharedBufferWithID
+	{
+		FSharedBufferWithID(FSharedBuffer InPayload);
+
+		FSharedBufferWithID() = default;
+		FSharedBufferWithID(FSharedBufferWithID&&) = default;
+		FSharedBufferWithID& operator=(FSharedBufferWithID&&) = default;
+
+		FSharedBufferWithID(const FSharedBufferWithID&) = delete;
+		FSharedBufferWithID& operator=(const FSharedBufferWithID&) = delete;
+
+	private:
+		friend FVirtualizedUntypedBulkData;
+
+		FSharedBuffer Payload;
+		FPayloadId PayloadId;
+	};
+
+	/**
+	 * Allows the existing payload to be replaced with a new one.
+	 *
+	 * 
+	 * To pass in a raw pointer use 'FSharedBuffer::...(Data, Size)' to create a valid FSharedBuffer.
+	 * Use 'FSharedBuffer::MakeView' if you want to retain ownership on the data being passed in, and use
+	 * 'FSharedBuffer::TakeOwnership' if you are okay with the bulkdata object taking over ownership of it.
+	 * The bulkdata object must own its internal buffer, so if you pass in a non-owned FSharedBuffer (ie
+	 * by using 'FSharedBuffer::MakeView') then a clone of the data will be created internally and assigned
+	 * to the bulkdata object.
+	 * 
+	 * Use this override if you want compute PayloadId before updating the bulkdata
+	 *
+	 * @param InPayload				The payload to update the bulkdata with
+	 * @param InCompressionFormat	The compression format to use, NAME_None indicates that the
+	 *								payload is already in a compressed format and will not gain from being 
+	 *								compressed again. These payloads will never be compressed. NAME_Default 
+	 *								will apply whichever compression format that the underlying code deems 
+	 *								appropriate. Other specific compression formats may be allowed, see the 
+	 *								documentation of FCompressedBuffer for details.
+	 */
+	void UpdatePayload(FSharedBufferWithID InPayload, FName CompressionFormat = NAME_Default);
+
 	/** 
 	 * Allows the compression format to be specified that will be applied the next time that the bulkdata
 	 * is saved. For non-virtualized payloads this will occur when the package is next saved. For 
@@ -235,6 +279,8 @@ private:
 
 	FRIEND_ENUM_CLASS_FLAGS(EFlags);
 
+	void UpdatePayloadImpl(FSharedBuffer&& InPayload, FPayloadId&& InPayloadID, FName CompressionFormat = NAME_Default);
+
 	FCompressedBuffer GetDataInternal() const;
 
 	FCompressedBuffer LoadFromDisk() const;
diff --git a/Engine/Source/Runtime/Engine/Classes/Engine/Texture.h b/Engine/Source/Runtime/Engine/Classes/Engine/Texture.h
--- a/Engine/Source/Runtime/Engine/Classes/Engine/Texture.h
+++ b/Engine/Source/Runtime/Engine/Classes/Engine/Texture.h
@@ -138,7 +138,7 @@ struct FTextureSource
 		const FTextureSourceBlock* InBlocks,
 		int32 InNumLayers,
 		int32 InNumBlocks,
-		FSharedBuffer NewData);
+		UE::Virtualization::FVirtualizedUntypedBulkData::FSharedBufferWithID NewData);
 
 	ENGINE_API void InitLayered(
 		int32 NewSizeX,
@@ -157,7 +157,7 @@ struct FTextureSource
 		int32 NewNumLayers,
 		int32 NewNumMips,
 		const ETextureSourceFormat* NewLayerFormat,
-		FSharedBuffer NewData
+		UE::Virtualization::FVirtualizedUntypedBulkData::FSharedBufferWithID NewData
 	);
 
 	/**
@@ -193,7 +193,7 @@ struct FTextureSource
 		int32 NewNumSlices,
 		int32 NewNumMips,
 		ETextureSourceFormat NewFormat,
-		FSharedBuffer NewData
+		UE::Virtualization::FVirtualizedUntypedBulkData::FSharedBufferWithID NewData
 	);
 
 	/**
diff --git a/Engine/Source/Runtime/Engine/Private/Texture.cpp b/Engine/Source/Runtime/Engine/Private/Texture.cpp
--- a/Engine/Source/Runtime/Engine/Private/Texture.cpp
+++ b/Engine/Source/Runtime/Engine/Private/Texture.cpp
@@ -1099,7 +1099,7 @@ void FTextureSource::InitBlocked(const ETextureSourceFormat* InLayerFormats,
 	const FTextureSourceBlock* InBlocks,
 	int32 InNumLayers,
 	int32 InNumBlocks,
-	FSharedBuffer NewData)
+	UE::Virtualization::FVirtualizedUntypedBulkData::FSharedBufferWithID NewData)
 {
 	InitBlockedImpl(InLayerFormats, InBlocks, InNumLayers, InNumBlocks);
 
@@ -1150,7 +1150,7 @@ void FTextureSource::InitLayered(
 	int32 NewNumLayers,
 	int32 NewNumMips,
 	const ETextureSourceFormat* NewLayerFormat,
-	FSharedBuffer NewData)
+	UE::Virtualization::FVirtualizedUntypedBulkData::FSharedBufferWithID NewData)
 {
 	InitLayeredImpl(
 		NewSizeX,
@@ -1182,10 +1182,10 @@ void FTextureSource::Init(
 	int32 NewNumSlices,
 	int32 NewNumMips,
 	ETextureSourceFormat NewFormat,
-	FSharedBuffer NewData
+	UE::Virtualization::FVirtualizedUntypedBulkData::FSharedBufferWithID NewData
 )
 {
-	InitLayered(NewSizeX, NewSizeY, NewNumSlices, 1, NewNumMips, &NewFormat, NewData);
+	InitLayered(NewSizeX, NewSizeY, NewNumSlices, 1, NewNumMips, &NewFormat, MoveTemp(NewData));
 }
 
 void FTextureSource::Init2DWithMipChain(
