Fixed a bad assertion in setup texture functions.

diff --git a/Engine/Plugins/Runtime/Interchange/Source/Import/Private/Texture/InterchangeTextureFactory.cpp b/Engine/Plugins/Runtime/Interchange/Source/Import/Private/Texture/InterchangeTextureFactory.cpp
--- a/Engine/Plugins/Runtime/Interchange/Source/Import/Private/Texture/InterchangeTextureFactory.cpp
+++ b/Engine/Plugins/Runtime/Interchange/Source/Import/Private/Texture/InterchangeTextureFactory.cpp
@@ -415,9 +415,11 @@ namespace UE::Interchange::Private::InterchangeTextureFactory
 				}
 			}
 		}
-
-		// The TexturePayload should be validated before calling this function
-		checkNoEntry();
+		else
+		{
+			// The TexturePayload should be validated before calling this function
+			checkNoEntry();
+		}
 	}
 
 	bool CanSetupTextureCubeSourceData(FTexturePayloadVariant& TexturePayload)
@@ -478,9 +480,11 @@ namespace UE::Interchange::Private::InterchangeTextureFactory
 				SetupTextureSourceData(TextureCube, LightProfile, MoveTemp(ProcessedPayload.PayloadAndId));
 			}
 		}
-
-		// The TexturePayload should be validated before calling this function
-		checkNoEntry();
+		else
+		{
+			// The TexturePayload should be validated before calling this function
+			checkNoEntry();
+		}
 	}
 
 	bool CanSetupTexture2DArraySourceData(FTexturePayloadVariant& TexturePayload)
@@ -537,9 +541,11 @@ namespace UE::Interchange::Private::InterchangeTextureFactory
 				SetupTextureSourceData(Texture2DArray, LightProfile, MoveTemp(ProcessedPayload.PayloadAndId));
 			}
 		}
-
-		// The TexturePayload should be validated before calling this function
-		checkNoEntry();
+		else
+		{
+			// The TexturePayload should be validated before calling this function
+			checkNoEntry();
+		}
 	}
 
 	void LogErrorInvalidPayload(const FString& TextureClass, const FString& ObjectName)
