Interchange: Spawn static mesh, skeletal mesh, light and camera actors from FBX
- Renamed InterchangeSceneNode::CustomMeshDependencyUid to AssetInstanceUid. It tells which asset, if any, a scene node is instanciating.
- Added an actor class attribute to the actor factory node. The pipeline must set this class to tell the factory which class to use when spawning the actor. By default, an AActor is spawned, if the attribute isn't set.
- Camera and lights now have asset nodes and are instanciated using scene nodes. This allows to reuse the same light or camera data multiple times in a scene.

diff --git a/Engine/Plugins/Runtime/Interchange/Source/Import/Private/Mesh/InterchangeSkeletalMeshFactory.cpp b/Engine/Plugins/Runtime/Interchange/Source/Import/Private/Mesh/InterchangeSkeletalMeshFactory.cpp
--- a/Engine/Plugins/Runtime/Interchange/Source/Import/Private/Mesh/InterchangeSkeletalMeshFactory.cpp
+++ b/Engine/Plugins/Runtime/Interchange/Source/Import/Private/Mesh/InterchangeSkeletalMeshFactory.cpp
@@ -828,7 +828,7 @@ UObject* UInterchangeSkeletalMeshFactory::CreateAsset(const FCreateAssetParams&
 								continue;
 							}
 							FString MeshDependencyUid;
-							MeshReference.SceneNode->GetCustomMeshDependencyUid(MeshDependencyUid);
+							MeshReference.SceneNode->GetCustomAssetInstanceUid(MeshDependencyUid);
 							MeshReference.MeshNode = Cast<UInterchangeMeshNode>(Arguments.NodeContainer->GetNode(MeshDependencyUid));
 							//Cache the scene node global matrix, we will use this matrix to bake the vertices
 							FTransform SceneNodeGlobalTransform;
diff --git a/Engine/Plugins/Runtime/Interchange/Source/Import/Private/Mesh/InterchangeStaticMeshFactory.cpp b/Engine/Plugins/Runtime/Interchange/Source/Import/Private/Mesh/InterchangeStaticMeshFactory.cpp
--- a/Engine/Plugins/Runtime/Interchange/Source/Import/Private/Mesh/InterchangeStaticMeshFactory.cpp
+++ b/Engine/Plugins/Runtime/Interchange/Source/Import/Private/Mesh/InterchangeStaticMeshFactory.cpp
@@ -196,7 +196,7 @@ UObject* UInterchangeStaticMeshFactory::CreateAsset(const FCreateAssetParams& Ar
 							}
 
 							FString MeshDependencyUid;
-							MeshReference.SceneNode->GetCustomMeshDependencyUid(MeshDependencyUid);
+							MeshReference.SceneNode->GetCustomAssetInstanceUid(MeshDependencyUid);
 							MeshReference.MeshNode = Cast<UInterchangeMeshNode>(Arguments.NodeContainer->GetNode(MeshDependencyUid));
 
 							// Cache the scene node global matrix, we will use this matrix to bake the vertices
diff --git a/Engine/Plugins/Runtime/Interchange/Source/Nodes/Private/InterchangeActorFactoryNode.cpp b/Engine/Plugins/Runtime/Interchange/Source/Nodes/Private/InterchangeActorFactoryNode.cpp
--- a/Engine/Plugins/Runtime/Interchange/Source/Nodes/Private/InterchangeActorFactoryNode.cpp
+++ b/Engine/Plugins/Runtime/Interchange/Source/Nodes/Private/InterchangeActorFactoryNode.cpp
@@ -9,6 +9,16 @@
 UClass* UInterchangeActorFactoryNode::GetObjectClass() const
 {
 #if WITH_ENGINE
+	FString ActorClassName;
+	if (GetCustomActorClassName(ActorClassName))
+	{
+		UClass* ActorClass = FindObject<UClass>(nullptr, *ActorClassName);
+		if (ActorClass->IsChildOf<AActor>())
+		{
+			return ActorClass;
+		}
+	}
+
 	return AActor::StaticClass();
 #else
 	return nullptr;
diff --git a/Engine/Plugins/Runtime/Interchange/Source/Nodes/Public/InterchangeActorFactoryNode.h b/Engine/Plugins/Runtime/Interchange/Source/Nodes/Public/InterchangeActorFactoryNode.h
--- a/Engine/Plugins/Runtime/Interchange/Source/Nodes/Public/InterchangeActorFactoryNode.h
+++ b/Engine/Plugins/Runtime/Interchange/Source/Nodes/Public/InterchangeActorFactoryNode.h
@@ -26,6 +26,19 @@ public:
 		IMPLEMENT_NODE_ATTRIBUTE_SETTER_NODELEGATE(GlobalTransform, FTransform);
 	}
 
+	UFUNCTION(BlueprintCallable, Category = "Interchange | Node | ActorFactory")
+	bool GetCustomActorClassName(FString& AttributeValue) const
+	{
+		IMPLEMENT_NODE_ATTRIBUTE_GETTER(ActorClassName, FString);
+	}
+
+	UFUNCTION(BlueprintCallable, Category = "Interchange | Node | ActorFactory")
+	bool SetCustomActorClassName(const FString& AttributeValue)
+	{
+		IMPLEMENT_NODE_ATTRIBUTE_SETTER_NODELEGATE(ActorClassName, FString);
+	}
+
 private:
 	const UE::Interchange::FAttributeKey Macro_CustomGlobalTransformKey = UE::Interchange::FAttributeKey(TEXT("GlobalTransform"));
+	const UE::Interchange::FAttributeKey Macro_CustomActorClassNameKey = UE::Interchange::FAttributeKey(TEXT("ActorClassName"));
 };
\ No newline at end of file
diff --git a/Engine/Plugins/Runtime/Interchange/Source/Nodes/Public/InterchangeCameraNode.h b/Engine/Plugins/Runtime/Interchange/Source/Nodes/Public/InterchangeCameraNode.h
new file mode 100644
--- /dev/null
+++ b/Engine/Plugins/Runtime/Interchange/Source/Nodes/Public/InterchangeCameraNode.h
@@ -0,0 +1,29 @@
+// Copyright Epic Games, Inc. All Rights Reserved.
+
+#pragma once
+
+#include "CoreMinimal.h"
+#include "Nodes/InterchangeBaseNode.h"
+
+#include "InterchangeCameraNode.generated.h"
+
+UCLASS(BlueprintType)
+class INTERCHANGENODES_API UInterchangeCameraNode : public UInterchangeBaseNode
+{
+	GENERATED_BODY()
+
+public:
+	static FStringView StaticAssetTypeName()
+	{
+		return TEXT("Camera");
+	}
+
+	/**
+	 * Return the node type name of the class, we use this when reporting errors
+	 */
+	virtual FString GetTypeName() const override
+	{
+		const FString TypeName = TEXT("CameraNode");
+		return TypeName;
+	}
+};
diff --git a/Engine/Plugins/Runtime/Interchange/Source/Nodes/Public/InterchangeLightNode.h b/Engine/Plugins/Runtime/Interchange/Source/Nodes/Public/InterchangeLightNode.h
new file mode 100644
--- /dev/null
+++ b/Engine/Plugins/Runtime/Interchange/Source/Nodes/Public/InterchangeLightNode.h
@@ -0,0 +1,30 @@
+// Copyright Epic Games, Inc. All Rights Reserved.
+
+#pragma once
+
+#include "CoreMinimal.h"
+#include "Nodes/InterchangeBaseNode.h"
+
+#include "InterchangeLightNode.generated.h"
+
+UCLASS(BlueprintType)
+class INTERCHANGENODES_API UInterchangeLightNode : public UInterchangeBaseNode
+{
+	GENERATED_BODY()
+
+public:
+	static FStringView StaticAssetTypeName()
+	{
+		return TEXT("Light");
+	}
+
+	/**
+	 * Return the node type name of the class, we use this when reporting errors
+	 */
+	virtual FString GetTypeName() const override
+	{
+		const FString TypeName = TEXT("LightNode");
+		return TypeName;
+	}
+
+};
diff --git a/Engine/Plugins/Runtime/Interchange/Source/Nodes/Public/InterchangeSceneNode.h b/Engine/Plugins/Runtime/Interchange/Source/Nodes/Public/InterchangeSceneNode.h
--- a/Engine/Plugins/Runtime/Interchange/Source/Nodes/Public/InterchangeSceneNode.h
+++ b/Engine/Plugins/Runtime/Interchange/Source/Nodes/Public/InterchangeSceneNode.h
@@ -7,7 +7,6 @@
 
 #include "InterchangeSceneNode.generated.h"
 
-//Interchange namespace
 //Interchange namespace
 namespace UE
 {
@@ -116,7 +115,7 @@ public:
 		}
 		else if (NodeAttributeKey == Macro_CustomLocalTransformKey
 				 || NodeAttributeKey == Macro_CustomGlobalTransformKey
-				 || NodeAttributeKey == Macro_CustomMeshDependencyUidKey)
+				 || NodeAttributeKey == Macro_CustomAssetInstanceUidKey)
 		{
 			return FString(TEXT("Scene"));
 		}
@@ -249,22 +248,23 @@ public:
 
 	/** Return false if the Attribute was not set previously.*/
 	UFUNCTION(BlueprintCallable, Category = "Interchange | Node | Scene")
-	bool GetCustomMeshDependencyUid(FString& AttributeValue) const
+	bool GetCustomAssetInstanceUid(FString& AttributeValue) const
 	{
-		IMPLEMENT_NODE_ATTRIBUTE_GETTER(MeshDependencyUid, FString);
+		IMPLEMENT_NODE_ATTRIBUTE_GETTER(AssetInstanceUid, FString);
 	}
 
+	/** Tells which asset, if any, a scene node is instantiating */
 	UFUNCTION(BlueprintCallable, Category = "Interchange | Node | Scene")
-	bool SetCustomMeshDependencyUid(const FString& AttributeValue)
+	bool SetCustomAssetInstanceUid(const FString& AttributeValue)
 	{
-		IMPLEMENT_NODE_ATTRIBUTE_SETTER_NODELEGATE(MeshDependencyUid, FString);
+		IMPLEMENT_NODE_ATTRIBUTE_SETTER_NODELEGATE(AssetInstanceUid, FString);
 	}
 
 private:
 	//Scene Attribute Keys
 	const UE::Interchange::FAttributeKey Macro_CustomLocalTransformKey = UE::Interchange::FAttributeKey(TEXT("LocalTransform"));
 	const UE::Interchange::FAttributeKey Macro_CustomGlobalTransformKey = UE::Interchange::FAttributeKey(TEXT("GlobalTransform"));
-	const UE::Interchange::FAttributeKey Macro_CustomMeshDependencyUidKey = UE::Interchange::FAttributeKey(TEXT("MeshDependencyUid"));
+	const UE::Interchange::FAttributeKey Macro_CustomAssetInstanceUidKey = UE::Interchange::FAttributeKey(TEXT("AssetInstanceUid"));
 
 	UE::Interchange::FNameAttributeArrayHelper NodeSpecializeTypes;
 	UE::Interchange::FNameAttributeArrayHelper MaterialDependencyUids;
diff --git a/Engine/Plugins/Runtime/Interchange/Source/Parsers/Fbx/Private/FbxAPI.cpp b/Engine/Plugins/Runtime/Interchange/Source/Parsers/Fbx/Private/FbxAPI.cpp
--- a/Engine/Plugins/Runtime/Interchange/Source/Parsers/Fbx/Private/FbxAPI.cpp
+++ b/Engine/Plugins/Runtime/Interchange/Source/Parsers/Fbx/Private/FbxAPI.cpp
@@ -3,9 +3,11 @@
 #include "FbxAPI.h"
 
 #include "CoreMinimal.h"
+#include "FbxCamera.h"
 #include "FbxConvert.h"
 #include "FbxHelper.h"
 #include "FbxInclude.h"
+#include "FbxLight.h"
 #include "FbxMaterial.h"
 #include "FbxMesh.h"
 #include "FbxScene.h"
@@ -109,6 +111,12 @@ namespace UE
 				FFbxMesh FbxMesh(*this);
 				FbxMesh.AddAllMeshes(SDKScene, SDKGeometryConverter, NodeContainer, PayloadContexts);
 
+				FFbxLight FbxLight(*this);
+				FbxLight.AddAllLights(SDKScene, NodeContainer);
+
+				FFbxCamera FbxCamera(*this);
+				FbxCamera.AddAllCameras(SDKScene, NodeContainer);
+
 				FFbxScene FbxScene(*this);
 				FbxScene.AddHierarchy(SDKScene, NodeContainer);
 			}
diff --git a/Engine/Plugins/Runtime/Interchange/Source/Parsers/Fbx/Private/FbxCamera.cpp b/Engine/Plugins/Runtime/Interchange/Source/Parsers/Fbx/Private/FbxCamera.cpp
new file mode 100644
--- /dev/null
+++ b/Engine/Plugins/Runtime/Interchange/Source/Parsers/Fbx/Private/FbxCamera.cpp
@@ -0,0 +1,74 @@
+// Copyright Epic Games, Inc. All Rights Reserved.
+
+#include "FbxCamera.h"
+
+#include "CoreMinimal.h"
+#include "FbxAPI.h"
+#include "FbxHelper.h"
+#include "FbxInclude.h"
+#include "InterchangeCameraNode.h"
+#include "InterchangeResultsContainer.h"
+#include "Nodes/InterchangeBaseNodeContainer.h"
+
+#define LOCTEXT_NAMESPACE "InterchangeFbxCamera"
+
+namespace UE
+{
+	namespace Interchange
+	{
+		namespace Private
+		{
+			UInterchangeCameraNode* FFbxCamera::CreateCameraNode(UInterchangeBaseNodeContainer& NodeContainer, const FString& NodeUid, const FString& NodeName)
+			{
+				UInterchangeCameraNode* CameraNode = NewObject<UInterchangeCameraNode>(&NodeContainer, NAME_None);
+				if (!ensure(CameraNode))
+				{
+					UInterchangeResultError_Generic* Message = Parser.AddMessage<UInterchangeResultError_Generic>();
+					Message->Text = LOCTEXT("CannotAllocateNode", "Cannot allocate a node when importing FBX.");
+					return nullptr;
+				}
+
+				CameraNode->InitializeNode(NodeUid, NodeName, EInterchangeNodeContainerType::NodeContainerType_TranslatedAsset);
+				NodeContainer.AddNode(CameraNode);
+
+				return CameraNode;
+			}
+
+			void FFbxCamera::AddCamerasRecursively(FbxNode* Node, UInterchangeBaseNodeContainer& NodeContainer)
+			{
+				int32 AttributeCount = Node->GetNodeAttributeCount();
+				for (int32 AttributeIndex = 0; AttributeIndex < AttributeCount; ++AttributeIndex)
+				{
+					FbxNodeAttribute* NodeAttribute = Node->GetNodeAttributeByIndex(AttributeIndex);
+
+					if (NodeAttribute->GetAttributeType() == FbxNodeAttribute::eCamera)
+					{
+						FString NodeName = FFbxHelper::GetNodeAttributeName(NodeAttribute, UInterchangeCameraNode::StaticAssetTypeName());
+						FString NodeUid = FFbxHelper::GetNodeAttributeUniqueID(NodeAttribute, UInterchangeCameraNode::StaticAssetTypeName());
+
+						UInterchangeCameraNode* CameraNode = Cast<UInterchangeCameraNode>(NodeContainer.GetNode(NodeUid));
+
+						if (!CameraNode)
+						{
+							CreateCameraNode(NodeContainer, NodeUid, NodeName);
+						}
+					}
+				}
+
+				const int32 ChildCount = Node->GetChildCount();
+				for (int32 ChildIndex = 0; ChildIndex < ChildCount; ++ChildIndex)
+				{
+					FbxNode* ChildNode = Node->GetChild(ChildIndex);
+					AddCamerasRecursively(ChildNode, NodeContainer);
+				}
+			}
+
+			void FFbxCamera::AddAllCameras(FbxScene* SDKScene, UInterchangeBaseNodeContainer& NodeContainer)
+			{
+				AddCamerasRecursively(SDKScene->GetRootNode(), NodeContainer);
+			}
+		} //ns Private
+	} //ns Interchange
+}//ns UE
+
+#undef LOCTEXT_NAMESPACE
diff --git a/Engine/Plugins/Runtime/Interchange/Source/Parsers/Fbx/Private/FbxCamera.h b/Engine/Plugins/Runtime/Interchange/Source/Parsers/Fbx/Private/FbxCamera.h
new file mode 100644
--- /dev/null
+++ b/Engine/Plugins/Runtime/Interchange/Source/Parsers/Fbx/Private/FbxCamera.h
@@ -0,0 +1,41 @@
+// Copyright Epic Games, Inc. All Rights Reserved.
+
+#pragma once
+
+#include "FbxInclude.h"
+
+class UInterchangeBaseNodeContainer;
+class UInterchangeCameraNode;
+
+namespace UE
+{
+	namespace Interchange
+	{
+		namespace Private
+		{
+			class FFbxParser;
+
+			class FFbxCamera
+			{
+			public:
+				explicit FFbxCamera(FFbxParser& InParser)
+					: Parser(InParser)
+				{}
+
+				/**
+				 * Create a UInterchangeCameraNode and add it to the NodeContainer for all FbxNodeAttributes of type eCamera that the fbx file contains.
+				 * 
+				 * @note - Any node that already exist in the NodeContainer will not be created or modified.
+				 */
+				void AddAllCameras(FbxScene* SDKScene, UInterchangeBaseNodeContainer& NodeContainer);
+
+			protected:
+				UInterchangeCameraNode* CreateCameraNode(UInterchangeBaseNodeContainer& NodeContainer, const FString& NodeUID, const FString& NodeName);
+				void AddCamerasRecursively(FbxNode* Node, UInterchangeBaseNodeContainer& NodeContainer);
+
+			private:
+				FFbxParser& Parser;
+			};
+		}//ns Private
+	}//ns Interchange
+}//ns UE
diff --git a/Engine/Plugins/Runtime/Interchange/Source/Parsers/Fbx/Private/FbxHelper.cpp b/Engine/Plugins/Runtime/Interchange/Source/Parsers/Fbx/Private/FbxHelper.cpp
--- a/Engine/Plugins/Runtime/Interchange/Source/Parsers/Fbx/Private/FbxHelper.cpp
+++ b/Engine/Plugins/Runtime/Interchange/Source/Parsers/Fbx/Private/FbxHelper.cpp
@@ -16,28 +16,17 @@ namespace UE
 		{
 			FString FFbxHelper::GetMeshName(FbxGeometryBase* Mesh)
 			{
-				FString MeshName = FFbxHelper::GetFbxObjectName(Mesh);
-				if (MeshName.IsEmpty())
+				FString DefaultPrefix;
+				if (Mesh->GetAttributeType() == FbxNodeAttribute::eMesh)
 				{
-					if (Mesh->GetNodeCount() > 0)
-					{
-						if (Mesh->GetAttributeType() == FbxNodeAttribute::eMesh)
-						{
-							MeshName = TEXT("Mesh_");
-						}
-						else if (Mesh->GetAttributeType() == FbxNodeAttribute::eShape)
-						{
-							MeshName = TEXT("Shape_");
-						}
-						MeshName += FFbxHelper::GetFbxObjectName(Mesh->GetNode(0));
-					}
-					else
-					{
-						uint64 UniqueFbxObjectID = Mesh->GetUniqueID();
-						MeshName += GetUniqueIDString(UniqueFbxObjectID);
-					}
+					DefaultPrefix = TEXT("Mesh");
+				}
+				else if (Mesh->GetAttributeType() == FbxNodeAttribute::eShape)
+				{
+					DefaultPrefix = TEXT("Shape");
 				}
-				return MeshName;
+
+				return GetNodeAttributeName(Mesh, DefaultPrefix);
 			}
 
 			FString FFbxHelper::GetMeshUniqueID(FbxGeometryBase* Mesh)
@@ -50,29 +39,59 @@ namespace UE
 				FString MeshUniqueID;
 				if (Mesh->GetAttributeType() == FbxNodeAttribute::eMesh)
 				{
-					MeshUniqueID = TEXT("\\Mesh\\");
+					MeshUniqueID = TEXT("Mesh");
 				}
 				else if (Mesh->GetAttributeType() == FbxNodeAttribute::eShape)
 				{
-					MeshUniqueID = TEXT("\\Shape\\");
+					MeshUniqueID = TEXT("Shape");
 				}
-				FString MeshName = FFbxHelper::GetFbxObjectName(Mesh);
-				if (MeshName.IsEmpty())
+
+				return GetNodeAttributeUniqueID(Mesh, MeshUniqueID);
+			}
+
+			FString FFbxHelper::GetNodeAttributeName(FbxNodeAttribute* NodeAttribute, const FStringView DefaultNamePrefix)
+			{
+				FString NodeAttributeName = FFbxHelper::GetFbxObjectName(NodeAttribute);
+				if (NodeAttributeName.IsEmpty())
 				{
-					if (Mesh->GetNodeCount() > 0)
+					if (NodeAttribute->GetNodeCount() > 0)
 					{
-						MeshUniqueID += FFbxHelper::GetFbxNodeHierarchyName(Mesh->GetNode(0));
+						NodeAttributeName = FString(DefaultNamePrefix) + TEXT("_") + FFbxHelper::GetFbxObjectName(NodeAttribute->GetNode(0));
 					}
 					else
 					{
-						MeshUniqueID += GetMeshName(Mesh);
+						uint64 UniqueFbxObjectID = NodeAttribute->GetUniqueID();
+						NodeAttributeName += GetUniqueIDString(UniqueFbxObjectID);
 					}
 				}
-				else
+				return NodeAttributeName;
+			}
+
+			FString FFbxHelper::GetNodeAttributeUniqueID(FbxNodeAttribute* NodeAttribute, const FStringView Prefix)
+			{
+				if (!NodeAttribute)
 				{
-					MeshUniqueID += MeshName;
+					return {};
 				}
-				return MeshUniqueID;
+
+				FString NodeAttributeUniqueID = FString(TEXT("\\")) + Prefix + TEXT("\\");
+				FString NodeAttributeName = FFbxHelper::GetFbxObjectName(NodeAttribute);
+
+				if (NodeAttributeName.IsEmpty())
+				{
+					if (NodeAttribute->GetNodeCount() > 0)
+					{
+						NodeAttributeName = FFbxHelper::GetFbxNodeHierarchyName(NodeAttribute->GetNode(0));
+					}
+					else
+					{
+						NodeAttributeName = GetNodeAttributeName(NodeAttribute, Prefix);
+					}
+				}
+
+				NodeAttributeUniqueID += NodeAttributeName;
+
+				return NodeAttributeUniqueID;
 			}
 
 			FString FFbxHelper::GetFbxObjectName(const FbxObject* Object)
diff --git a/Engine/Plugins/Runtime/Interchange/Source/Parsers/Fbx/Private/FbxHelper.h b/Engine/Plugins/Runtime/Interchange/Source/Parsers/Fbx/Private/FbxHelper.h
--- a/Engine/Plugins/Runtime/Interchange/Source/Parsers/Fbx/Private/FbxHelper.h
+++ b/Engine/Plugins/Runtime/Interchange/Source/Parsers/Fbx/Private/FbxHelper.h
@@ -28,6 +28,9 @@ namespace UE
 				static FString GetMeshName(FbxGeometryBase* Mesh);
 				static FString GetMeshUniqueID(FbxGeometryBase* Mesh);
 
+				static FString GetNodeAttributeName(FbxNodeAttribute* NodeAttribute, const FStringView DefaultNamePrefix);
+				static FString GetNodeAttributeUniqueID(FbxNodeAttribute* NodeAttribute, const FStringView Prefix);
+
 				/**
 				 * Return the name of an FbxObject, return empty string if the object is null.
 				 */
diff --git a/Engine/Plugins/Runtime/Interchange/Source/Parsers/Fbx/Private/FbxLight.cpp b/Engine/Plugins/Runtime/Interchange/Source/Parsers/Fbx/Private/FbxLight.cpp
new file mode 100644
--- /dev/null
+++ b/Engine/Plugins/Runtime/Interchange/Source/Parsers/Fbx/Private/FbxLight.cpp
@@ -0,0 +1,74 @@
+// Copyright Epic Games, Inc. All Rights Reserved.
+
+#include "FbxLight.h"
+
+#include "CoreMinimal.h"
+#include "FbxAPI.h"
+#include "FbxHelper.h"
+#include "FbxInclude.h"
+#include "InterchangeLightNode.h"
+#include "InterchangeResultsContainer.h"
+#include "Nodes/InterchangeBaseNodeContainer.h"
+
+#define LOCTEXT_NAMESPACE "InterchangeFbxLight"
+
+namespace UE
+{
+	namespace Interchange
+	{
+		namespace Private
+		{
+			UInterchangeLightNode* FFbxLight::CreateLightNode(UInterchangeBaseNodeContainer& NodeContainer, const FString& NodeUid, const FString& NodeName)
+			{
+				UInterchangeLightNode* LightNode = NewObject<UInterchangeLightNode>(&NodeContainer, NAME_None);
+				if (!ensure(LightNode))
+				{
+					UInterchangeResultError_Generic* Message = Parser.AddMessage<UInterchangeResultError_Generic>();
+					Message->Text = LOCTEXT("CannotAllocateNode", "Cannot allocate a node when importing FBX.");
+					return nullptr;
+				}
+
+				LightNode->InitializeNode(NodeUid, NodeName, EInterchangeNodeContainerType::NodeContainerType_TranslatedAsset);
+				NodeContainer.AddNode(LightNode);
+
+				return LightNode;
+			}
+
+			void FFbxLight::AddLightsRecursively(FbxNode* Node, UInterchangeBaseNodeContainer& NodeContainer)
+			{
+				int32 AttributeCount = Node->GetNodeAttributeCount();
+				for (int32 AttributeIndex = 0; AttributeIndex < AttributeCount; ++AttributeIndex)
+				{
+					FbxNodeAttribute* NodeAttribute = Node->GetNodeAttributeByIndex(AttributeIndex);
+
+					if (NodeAttribute->GetAttributeType() == FbxNodeAttribute::eLight)
+					{
+						FString NodeName = FFbxHelper::GetNodeAttributeName(NodeAttribute, UInterchangeLightNode::StaticAssetTypeName());
+						FString NodeUid = FFbxHelper::GetNodeAttributeUniqueID(NodeAttribute, UInterchangeLightNode::StaticAssetTypeName());
+
+						UInterchangeLightNode* LightNode = Cast<UInterchangeLightNode>(NodeContainer.GetNode(NodeUid));
+
+						if (!LightNode)
+						{
+							CreateLightNode(NodeContainer, NodeUid, NodeName);
+						}
+					}
+				}
+
+				const int32 ChildCount = Node->GetChildCount();
+				for (int32 ChildIndex = 0; ChildIndex < ChildCount; ++ChildIndex)
+				{
+					FbxNode* ChildNode = Node->GetChild(ChildIndex);
+					AddLightsRecursively(ChildNode, NodeContainer);
+				}
+			}
+
+			void FFbxLight::AddAllLights(FbxScene* SDKScene, UInterchangeBaseNodeContainer& NodeContainer)
+			{
+				AddLightsRecursively(SDKScene->GetRootNode(), NodeContainer);
+			}
+		} //ns Private
+	} //ns Interchange
+}//ns UE
+
+#undef LOCTEXT_NAMESPACE
diff --git a/Engine/Plugins/Runtime/Interchange/Source/Parsers/Fbx/Private/FbxLight.h b/Engine/Plugins/Runtime/Interchange/Source/Parsers/Fbx/Private/FbxLight.h
new file mode 100644
--- /dev/null
+++ b/Engine/Plugins/Runtime/Interchange/Source/Parsers/Fbx/Private/FbxLight.h
@@ -0,0 +1,41 @@
+// Copyright Epic Games, Inc. All Rights Reserved.
+
+#pragma once
+
+#include "FbxInclude.h"
+
+class UInterchangeBaseNodeContainer;
+class UInterchangeLightNode;
+
+namespace UE
+{
+	namespace Interchange
+	{
+		namespace Private
+		{
+			class FFbxParser;
+
+			class FFbxLight
+			{
+			public:
+				explicit FFbxLight(FFbxParser& InParser)
+					: Parser(InParser)
+				{}
+
+				/**
+				 * Create a UInterchangeLightNode and add it to the NodeContainer for all FbxNodeAttributes of type eLight the fbx file contains.
+				 * 
+				 * @note - Any node that already exist in the NodeContainer will not be created or modified.
+				 */
+				void AddAllLights(FbxScene* SDKScene, UInterchangeBaseNodeContainer& NodeContainer);
+
+			protected:
+				UInterchangeLightNode* CreateLightNode(UInterchangeBaseNodeContainer& NodeContainer, const FString& NodeUID, const FString& NodeName);
+				void AddLightsRecursively(FbxNode* Node, UInterchangeBaseNodeContainer& NodeContainer);
+
+			private:
+				FFbxParser& Parser;
+			};
+		}//ns Private
+	}//ns Interchange
+}//ns UE
diff --git a/Engine/Plugins/Runtime/Interchange/Source/Parsers/Fbx/Private/FbxScene.cpp b/Engine/Plugins/Runtime/Interchange/Source/Parsers/Fbx/Private/FbxScene.cpp
--- a/Engine/Plugins/Runtime/Interchange/Source/Parsers/Fbx/Private/FbxScene.cpp
+++ b/Engine/Plugins/Runtime/Interchange/Source/Parsers/Fbx/Private/FbxScene.cpp
@@ -9,6 +9,8 @@
 #include "FbxInclude.h"
 #include "FbxMaterial.h"
 #include "FbxMesh.h"
+#include "InterchangeCameraNode.h"
+#include "InterchangeLightNode.h"
 #include "InterchangeMeshNode.h"
 #include "InterchangeResultsContainer.h"
 #include "InterchangeSceneNode.h"
@@ -23,7 +25,7 @@ namespace UE
 		namespace Private
 		{
 
-			void FFbxScene::CreateMeshNodeReference(UInterchangeSceneNode* UnrealSceneNode, FbxNode* FbxSceneNode, FbxNodeAttribute* NodeAttribute, UInterchangeBaseNodeContainer& NodeContainer)
+			void FFbxScene::CreateMeshNodeReference(UInterchangeSceneNode* UnrealSceneNode, FbxNodeAttribute* NodeAttribute, UInterchangeBaseNodeContainer& NodeContainer)
 			{
 				UInterchangeMeshNode* MeshNode = nullptr;
 				if (NodeAttribute->GetAttributeType() == FbxNodeAttribute::eMesh)
@@ -42,18 +44,29 @@ namespace UE
 
 				if (MeshNode)
 				{
-					UnrealSceneNode->SetCustomMeshDependencyUid(MeshNode->GetUniqueID());
+					UnrealSceneNode->SetCustomAssetInstanceUid(MeshNode->GetUniqueID());
 					MeshNode->SetSceneInstanceUid(UnrealSceneNode->GetUniqueID());
 				}
 			}
 
-			void FFbxScene::CreateCameraNodeReference(UInterchangeSceneNode* UnrealSceneNode, FbxNode* FbxSceneNode, FbxNodeAttribute* NodeAttribute, UInterchangeBaseNodeContainer& NodeContainer)
+			void CreateAssetNodeReference(UInterchangeSceneNode* UnrealSceneNode, FbxNodeAttribute* NodeAttribute, UInterchangeBaseNodeContainer& NodeContainer, const FStringView TypeName)
 			{
+				const FString AssetUniqueID = FFbxHelper::GetNodeAttributeUniqueID(NodeAttribute, TypeName);
+
+				if (UInterchangeBaseNode* AssetNode = NodeContainer.GetNode(AssetUniqueID))
+				{
+					UnrealSceneNode->SetCustomAssetInstanceUid(AssetNode->GetUniqueID());
+				}
 			}
 
-			void FFbxScene::CreateLightNodeReference(UInterchangeSceneNode* UnrealSceneNode, FbxNode* FbxSceneNode, FbxNodeAttribute* NodeAttribute, UInterchangeBaseNodeContainer& NodeContainer)
+			void FFbxScene::CreateCameraNodeReference(UInterchangeSceneNode* UnrealSceneNode, FbxNodeAttribute* NodeAttribute, UInterchangeBaseNodeContainer& NodeContainer)
 			{
+				CreateAssetNodeReference(UnrealSceneNode, NodeAttribute, NodeContainer, UInterchangeCameraNode::StaticAssetTypeName());
+			}
 
+			void FFbxScene::CreateLightNodeReference(UInterchangeSceneNode* UnrealSceneNode, FbxNodeAttribute* NodeAttribute, UInterchangeBaseNodeContainer& NodeContainer)
+			{
+				CreateAssetNodeReference(UnrealSceneNode, NodeAttribute, NodeContainer, UInterchangeLightNode::StaticAssetTypeName());
 			}
 
 			UInterchangeSceneNode* FFbxScene::AddHierarchyRecursively(FbxNode* Node, FbxScene* SDKScene, UInterchangeBaseNodeContainer& NodeContainer)
@@ -125,7 +138,7 @@ namespace UE
 							//For Mesh attribute we add the fbx nodes materials
 							FFbxMaterial FbxMaterial(Parser);
 							FbxMaterial.AddAllNodeMaterials(UnrealNode, Node, NodeContainer);
-							CreateMeshNodeReference(UnrealNode, Node, NodeAttribute, NodeContainer);
+							CreateMeshNodeReference(UnrealNode, NodeAttribute, NodeContainer);
 							break;
 						}
 						case FbxNodeAttribute::eLODGroup:
@@ -136,13 +149,13 @@ namespace UE
 						case FbxNodeAttribute::eCamera:
 						{
 							//Add the Camera asset
-							CreateCameraNodeReference(UnrealNode, Node, NodeAttribute, NodeContainer);
+							CreateCameraNodeReference(UnrealNode, NodeAttribute, NodeContainer);
 							break;
 						}
 						case FbxNodeAttribute::eLight:
 						{
 							//Add the Light asset
-							CreateLightNodeReference(UnrealNode, Node, NodeAttribute, NodeContainer);
+							CreateLightNodeReference(UnrealNode, NodeAttribute, NodeContainer);
 							break;
 						}
 					}
diff --git a/Engine/Plugins/Runtime/Interchange/Source/Parsers/Fbx/Private/FbxScene.h b/Engine/Plugins/Runtime/Interchange/Source/Parsers/Fbx/Private/FbxScene.h
--- a/Engine/Plugins/Runtime/Interchange/Source/Parsers/Fbx/Private/FbxScene.h
+++ b/Engine/Plugins/Runtime/Interchange/Source/Parsers/Fbx/Private/FbxScene.h
@@ -28,9 +28,9 @@ namespace UE
 				UInterchangeSceneNode* CreateTransformNode(UInterchangeBaseNodeContainer& NodeContainer, const FString& NodeName, const FString& NodeUniqueID);
 
 			protected:
-				void CreateMeshNodeReference(UInterchangeSceneNode* UnrealSceneNode, FbxNode* FbxSceneNode, FbxNodeAttribute* NodeAttribute, UInterchangeBaseNodeContainer& NodeContainer);
-				void CreateCameraNodeReference(UInterchangeSceneNode* UnrealSceneNode, FbxNode* FbxSceneNode, FbxNodeAttribute* NodeAttribute, UInterchangeBaseNodeContainer& NodeContainer);
-				void CreateLightNodeReference(UInterchangeSceneNode* UnrealSceneNode, FbxNode* FbxSceneNode, FbxNodeAttribute* NodeAttribute, UInterchangeBaseNodeContainer& NodeContainer);
+				void CreateMeshNodeReference(UInterchangeSceneNode* UnrealSceneNode, FbxNodeAttribute* NodeAttribute, UInterchangeBaseNodeContainer& NodeContainer);
+				void CreateCameraNodeReference(UInterchangeSceneNode* UnrealSceneNode, FbxNodeAttribute* NodeAttribute, UInterchangeBaseNodeContainer& NodeContainer);
+				void CreateLightNodeReference(UInterchangeSceneNode* UnrealSceneNode, FbxNodeAttribute* NodeAttribute, UInterchangeBaseNodeContainer& NodeContainer);
 				UInterchangeSceneNode* AddHierarchyRecursively(FbxNode* Node, FbxScene* SDKScene, UInterchangeBaseNodeContainer& NodeContainer);
 
 			private:
diff --git a/Engine/Plugins/Runtime/Interchange/Source/Pipelines/InterchangePipelines.Build.cs b/Engine/Plugins/Runtime/Interchange/Source/Pipelines/InterchangePipelines.Build.cs
--- a/Engine/Plugins/Runtime/Interchange/Source/Pipelines/InterchangePipelines.Build.cs
+++ b/Engine/Plugins/Runtime/Interchange/Source/Pipelines/InterchangePipelines.Build.cs
@@ -18,6 +18,13 @@ namespace UnrealBuildTool.Rules
 				}
 			);
 
+			PrivateDependencyModuleNames.AddRange(
+				new string[]
+				{
+					"CinematicCamera",
+				}
+			);
+
 			if (Target.Type == TargetType.Editor)
 			{
 				PrivateDependencyModuleNames.AddRange(
diff --git a/Engine/Plugins/Runtime/Interchange/Source/Pipelines/Private/InterchangeGenericAssetsPipeline.cpp b/Engine/Plugins/Runtime/Interchange/Source/Pipelines/Private/InterchangeGenericAssetsPipeline.cpp
--- a/Engine/Plugins/Runtime/Interchange/Source/Pipelines/Private/InterchangeGenericAssetsPipeline.cpp
+++ b/Engine/Plugins/Runtime/Interchange/Source/Pipelines/Private/InterchangeGenericAssetsPipeline.cpp
@@ -449,7 +449,7 @@ bool UInterchangeGenericAssetsPipeline::MakeMeshFactoryNodeUidAndDisplayLabel(co
 				if (UInterchangeSceneNode* SceneNode = Cast<UInterchangeSceneNode>(Node))
 				{
 					FString RefMeshUid;
-					if (SceneNode->GetCustomMeshDependencyUid(RefMeshUid))
+					if (SceneNode->GetCustomAssetInstanceUid(RefMeshUid))
 					{
 						UInterchangeBaseNode* MeshNode = BaseNodeContainer->GetNode(RefMeshUid);
 						if (MeshNode)
@@ -582,7 +582,7 @@ void UInterchangeGenericAssetsPipeline::AddLodDataToStaticMesh(UInterchangeStati
 			if (const UInterchangeSceneNode* SceneNode = Cast<UInterchangeSceneNode>(BaseNodeContainer->GetNode(NodeUid)))
 			{
 				FString MeshDependency;
-				SceneNode->GetCustomMeshDependencyUid(MeshDependency);
+				SceneNode->GetCustomAssetInstanceUid(MeshDependency);
 				if (BaseNodeContainer->IsNodeUidValid(MeshDependency))
 				{
 					StaticMeshFactoryNode->AddTargetAssetUid(MeshDependency);
diff --git a/Engine/Plugins/Runtime/Interchange/Source/Pipelines/Private/InterchangeGenericScenesPipeline.cpp b/Engine/Plugins/Runtime/Interchange/Source/Pipelines/Private/InterchangeGenericScenesPipeline.cpp
--- a/Engine/Plugins/Runtime/Interchange/Source/Pipelines/Private/InterchangeGenericScenesPipeline.cpp
+++ b/Engine/Plugins/Runtime/Interchange/Source/Pipelines/Private/InterchangeGenericScenesPipeline.cpp
@@ -3,9 +3,17 @@
 #include "InterchangeGenericScenesPipeline.h"
 
 #include "InterchangeActorFactoryNode.h"
+#include "InterchangeCameraNode.h"
+#include "InterchangeLightNode.h"
+#include "InterchangeMeshNode.h"
 #include "InterchangePipelineLog.h"
 #include "InterchangeSceneNode.h"
 
+#include "Animation/SkeletalMeshActor.h"
+#include "CineCameraActor.h"
+#include "Engine/PointLight.h"
+#include "Engine/StaticMeshActor.h"
+
 bool UInterchangeGenericLevelPipeline::ExecutePreImportPipeline(UInterchangeBaseNodeContainer* InBaseNodeContainer, const TArray<UInterchangeSourceData*>& InSourceDatas)
 {
 	if (!InBaseNodeContainer)
@@ -34,7 +42,11 @@ bool UInterchangeGenericLevelPipeline::ExecutePreImportPipeline(UInterchangeBase
 
 	for (const UInterchangeSceneNode* SceneNode : SceneNodes)
 	{
-		CreateActorFactoryNode(SceneNode, InBaseNodeContainer);
+		//Ignore specialized types for now as they are used for bones hierarchies and other asset internal data
+		if (SceneNode && SceneNode->GetSpecializedTypeCount() == 0)
+		{
+			CreateActorFactoryNode(SceneNode, InBaseNodeContainer);
+		}
 	}
 
 	return true;
@@ -67,6 +79,32 @@ void UInterchangeGenericLevelPipeline::CreateActorFactoryNode(const UInterchange
 		ActorFactoryNode->SetCustomGlobalTransform(GlobalTransform);
 	}
 	
+	FString AssetInstanceUid;
+	if (SceneNode->GetCustomAssetInstanceUid(AssetInstanceUid))
+	{
+		UInterchangeBaseNode* TranslatedAssetNode = FactoryNodeContainer->GetNode(AssetInstanceUid);
+
+		if (UInterchangeMeshNode* MeshNode = Cast<UInterchangeMeshNode>(TranslatedAssetNode))
+		{
+			if (MeshNode->IsSkinnedMesh())
+			{
+				ActorFactoryNode->SetCustomActorClassName(ASkeletalMeshActor::StaticClass()->GetPathName());
+			}
+			else
+			{
+				ActorFactoryNode->SetCustomActorClassName(AStaticMeshActor::StaticClass()->GetPathName());
+			}
+		}
+		else if (UInterchangeLightNode* LightNode = Cast<UInterchangeLightNode>(TranslatedAssetNode))
+		{
+			ActorFactoryNode->SetCustomActorClassName(APointLight::StaticClass()->GetPathName());
+		}
+		else if (UInterchangeCameraNode* CameraNode = Cast<UInterchangeCameraNode>(TranslatedAssetNode))
+		{
+			ActorFactoryNode->SetCustomActorClassName(ACineCameraActor::StaticClass()->GetPathName());
+		}
+	}
+
 	FactoryNodeContainer->AddNode(ActorFactoryNode);
 }
 
diff --git a/Engine/Plugins/Runtime/Interchange/Source/Pipelines/Private/InterchangeGenericSkeletalMeshPipeline.cpp b/Engine/Plugins/Runtime/Interchange/Source/Pipelines/Private/InterchangeGenericSkeletalMeshPipeline.cpp
--- a/Engine/Plugins/Runtime/Interchange/Source/Pipelines/Private/InterchangeGenericSkeletalMeshPipeline.cpp
+++ b/Engine/Plugins/Runtime/Interchange/Source/Pipelines/Private/InterchangeGenericSkeletalMeshPipeline.cpp
@@ -512,7 +512,7 @@ UInterchangeSkeletalMeshFactoryNode* UInterchangeGenericAssetsPipeline::CreateSk
 					if (SceneNode)
 					{
 						FString MeshNodeUid;
-						if (SceneNode->GetCustomMeshDependencyUid(MeshNodeUid))
+						if (SceneNode->GetCustomAssetInstanceUid(MeshNodeUid))
 						{
 							OutSceneNodeCount = MeshUids.Num();
 							OutFirstMeshNodeUid = MeshNodeUid;
diff --git a/Engine/Plugins/Runtime/Interchange/Source/Pipelines/Private/InterchangePipelineMeshesUtilities.cpp b/Engine/Plugins/Runtime/Interchange/Source/Pipelines/Private/InterchangePipelineMeshesUtilities.cpp
--- a/Engine/Plugins/Runtime/Interchange/Source/Pipelines/Private/InterchangePipelineMeshesUtilities.cpp
+++ b/Engine/Plugins/Runtime/Interchange/Source/Pipelines/Private/InterchangePipelineMeshesUtilities.cpp
@@ -49,69 +49,72 @@ UInterchangePipelineMeshesUtilities* UInterchangePipelineMeshesUtilities::Create
 				}
 
 				FString MeshUid;
-				if (SceneNode->GetCustomMeshDependencyUid(MeshUid))
+				if (SceneNode->GetCustomAssetInstanceUid(MeshUid))
 				{
-					UInterchangeSceneNode* ParentMeshSceneNode = Cast<UInterchangeSceneNode>(BaseNodeContainer->GetNode(SceneNode->GetParentUid()));
-					if (ParentMeshSceneNode)
+					if (BaseNodeContainer->GetNode(MeshUid)->IsA<UInterchangeMeshNode>())
 					{
-						UInterchangeSceneNode* LodGroupNode = nullptr;
-						int32 LodIndex = 0;
-						FString LastChildUid = SceneNode->GetUniqueID();
-						do
+						UInterchangeSceneNode* ParentMeshSceneNode = Cast<UInterchangeSceneNode>(BaseNodeContainer->GetNode(SceneNode->GetParentUid()));
+						if (ParentMeshSceneNode)
 						{
-							if (ParentMeshSceneNode->IsSpecializedTypeContains(UE::Interchange::FSceneNodeStaticData::GetLodGroupSpecializeTypeString()))
+							UInterchangeSceneNode* LodGroupNode = nullptr;
+							int32 LodIndex = 0;
+							FString LastChildUid = SceneNode->GetUniqueID();
+							do
 							{
-								LodGroupNode = ParentMeshSceneNode;
-								TArray<FString> LodGroupChildrens = BaseNodeContainer->GetNodeChildrenUids(ParentMeshSceneNode->GetUniqueID());
-								for (int32 ChildLodIndex = 0; ChildLodIndex < LodGroupChildrens.Num(); ++ChildLodIndex)
+								if (ParentMeshSceneNode->IsSpecializedTypeContains(UE::Interchange::FSceneNodeStaticData::GetLodGroupSpecializeTypeString()))
 								{
-									const FString& ChildrenUid = LodGroupChildrens[ChildLodIndex];
-									if (ChildrenUid.Equals(LastChildUid))
+									LodGroupNode = ParentMeshSceneNode;
+									TArray<FString> LodGroupChildrens = BaseNodeContainer->GetNodeChildrenUids(ParentMeshSceneNode->GetUniqueID());
+									for (int32 ChildLodIndex = 0; ChildLodIndex < LodGroupChildrens.Num(); ++ChildLodIndex)
 									{
-										LodIndex = ChildLodIndex;
-										break;
+										const FString& ChildrenUid = LodGroupChildrens[ChildLodIndex];
+										if (ChildrenUid.Equals(LastChildUid))
+										{
+											LodIndex = ChildLodIndex;
+											break;
+										}
 									}
+									break;
 								}
-								break;
-							}
-							LastChildUid = ParentMeshSceneNode->GetUniqueID();
-							ParentMeshSceneNode = Cast<UInterchangeSceneNode>(BaseNodeContainer->GetNode(ParentMeshSceneNode->GetParentUid()));
-						} while (ParentMeshSceneNode);
+								LastChildUid = ParentMeshSceneNode->GetUniqueID();
+								ParentMeshSceneNode = Cast<UInterchangeSceneNode>(BaseNodeContainer->GetNode(ParentMeshSceneNode->GetParentUid()));
+							} while (ParentMeshSceneNode);
 
-						FInterchangeMeshGeometry& MeshGeometry = PipelineMeshesUtilities->MeshGeometriesPerMeshUid.FindChecked(MeshUid);
-						if (LodGroupNode)
-						{
-							//We have a LOD
-							FInterchangeMeshInstance& MeshInstance = PipelineMeshesUtilities->MeshInstancesPerMeshInstanceUid.FindOrAdd(LodGroupNode->GetUniqueID());
-							if (MeshInstance.LodGroupNode != nullptr)
+							FInterchangeMeshGeometry& MeshGeometry = PipelineMeshesUtilities->MeshGeometriesPerMeshUid.FindChecked(MeshUid);
+							if (LodGroupNode)
 							{
-								//This LodGroup was already created, verify everything is ok
-								checkSlow(MeshInstance.LodGroupNode == LodGroupNode);
-								checkSlow(MeshInstance.MeshInstanceUid.Equals(LodGroupNode->GetUniqueID()));
+								//We have a LOD
+								FInterchangeMeshInstance& MeshInstance = PipelineMeshesUtilities->MeshInstancesPerMeshInstanceUid.FindOrAdd(LodGroupNode->GetUniqueID());
+								if (MeshInstance.LodGroupNode != nullptr)
+								{
+									//This LodGroup was already created, verify everything is ok
+									checkSlow(MeshInstance.LodGroupNode == LodGroupNode);
+									checkSlow(MeshInstance.MeshInstanceUid.Equals(LodGroupNode->GetUniqueID()));
+								}
+								else
+								{
+									MeshInstance.LodGroupNode = LodGroupNode;
+									MeshInstance.MeshInstanceUid = LodGroupNode->GetUniqueID();
+								}
+								FInterchangeLodSceneNodeContainer& InstancedSceneNodes = MeshInstance.SceneNodePerLodIndex.FindOrAdd(LodIndex);
+								InstancedSceneNodes.SceneNodes.AddUnique(SceneNode);
+								MeshGeometry.ReferencingMeshInstanceUids.Add(MeshInstance.MeshInstanceUid);
+								MeshInstance.ReferencingMeshGeometryUids.Add(MeshUid);
+								MeshInstance.bReferenceSkinnedMesh |= MeshGeometry.MeshNode->IsSkinnedMesh();
+								MeshInstance.bReferenceBlendShape |= MeshGeometry.MeshNode->IsBlendShape();
 							}
 							else
 							{
-								MeshInstance.LodGroupNode = LodGroupNode;
-								MeshInstance.MeshInstanceUid = LodGroupNode->GetUniqueID();
+								FInterchangeMeshInstance& MeshInstance = PipelineMeshesUtilities->MeshInstancesPerMeshInstanceUid.FindOrAdd(NodeUid);
+								MeshInstance.LodGroupNode = nullptr;
+								MeshInstance.MeshInstanceUid = NodeUid;
+								FInterchangeLodSceneNodeContainer& InstancedSceneNodes = MeshInstance.SceneNodePerLodIndex.FindOrAdd(LodIndex);
+								InstancedSceneNodes.SceneNodes.AddUnique(SceneNode);
+								MeshGeometry.ReferencingMeshInstanceUids.Add(MeshInstance.MeshInstanceUid);
+								MeshInstance.ReferencingMeshGeometryUids.Add(MeshUid);
+								MeshInstance.bReferenceSkinnedMesh |= MeshGeometry.MeshNode->IsSkinnedMesh();
+								MeshInstance.bReferenceBlendShape |= MeshGeometry.MeshNode->IsBlendShape();
 							}
-							FInterchangeLodSceneNodeContainer& InstancedSceneNodes = MeshInstance.SceneNodePerLodIndex.FindOrAdd(LodIndex);
-							InstancedSceneNodes.SceneNodes.AddUnique(SceneNode);
-							MeshGeometry.ReferencingMeshInstanceUids.Add(MeshInstance.MeshInstanceUid);
-							MeshInstance.ReferencingMeshGeometryUids.Add(MeshUid);
-							MeshInstance.bReferenceSkinnedMesh |= MeshGeometry.MeshNode->IsSkinnedMesh();
-							MeshInstance.bReferenceBlendShape |= MeshGeometry.MeshNode->IsBlendShape();
-						}
-						else
-						{
-							FInterchangeMeshInstance& MeshInstance = PipelineMeshesUtilities->MeshInstancesPerMeshInstanceUid.FindOrAdd(NodeUid);
-							MeshInstance.LodGroupNode = nullptr;
-							MeshInstance.MeshInstanceUid = NodeUid;
-							FInterchangeLodSceneNodeContainer& InstancedSceneNodes = MeshInstance.SceneNodePerLodIndex.FindOrAdd(LodIndex);
-							InstancedSceneNodes.SceneNodes.AddUnique(SceneNode);
-							MeshGeometry.ReferencingMeshInstanceUids.Add(MeshInstance.MeshInstanceUid);
-							MeshInstance.ReferencingMeshGeometryUids.Add(MeshUid);
-							MeshInstance.bReferenceSkinnedMesh |= MeshGeometry.MeshNode->IsSkinnedMesh();
-							MeshInstance.bReferenceBlendShape |= MeshGeometry.MeshNode->IsBlendShape();
 						}
 					}
 				}
@@ -439,7 +442,7 @@ FString UInterchangePipelineMeshesUtilities::GetMeshInstanceSkeletonRootUid(cons
 	{
 		const UInterchangeSceneNode* SceneNode = MeshInstance.SceneNodePerLodIndex[BaseLodIndex].SceneNodes[0];
 		FString MeshNodeUid;
-		if (SceneNode->GetCustomMeshDependencyUid(MeshNodeUid))
+		if (SceneNode->GetCustomAssetInstanceUid(MeshNodeUid))
 		{
 			if (SkeletonRootUidPerMeshUid.Contains(MeshNodeUid))
 			{
diff --git a/Engine/Source/Runtime/Interchange/Engine/Private/Tasks/InterchangeTaskCreateSceneObjects.cpp b/Engine/Source/Runtime/Interchange/Engine/Private/Tasks/InterchangeTaskCreateSceneObjects.cpp
--- a/Engine/Source/Runtime/Interchange/Engine/Private/Tasks/InterchangeTaskCreateSceneObjects.cpp
+++ b/Engine/Source/Runtime/Interchange/Engine/Private/Tasks/InterchangeTaskCreateSceneObjects.cpp
@@ -85,6 +85,7 @@ void UE::Interchange::FTaskCreateSceneObjects::DoTask(ENamedThreads::Type Curren
 					UE::Interchange::FImportAsyncHelper::FImportedObjectInfo& ObjectInfo = ImportedInfos.AddDefaulted_GetRef();
 					ObjectInfo.ImportedObject = NodeObject;
 					ObjectInfo.Factory = Factory;
+					ObjectInfo.FactoryNode = Node;
 				}
 
 				if (CreateSceneObjectsParams.NodeContainer)
diff --git a/Engine/Source/Runtime/Interchange/Engine/Public/InterchangeManager.h b/Engine/Source/Runtime/Interchange/Engine/Public/InterchangeManager.h
--- a/Engine/Source/Runtime/Interchange/Engine/Public/InterchangeManager.h
+++ b/Engine/Source/Runtime/Interchange/Engine/Public/InterchangeManager.h
@@ -182,7 +182,7 @@ namespace UE
 			{
 				UObject* ImportedObject = nullptr; // The object that was imported
 				UInterchangeFactoryBase* Factory = nullptr; //The factory that created the imported object
-				UInterchangeBaseNode* FactoryNode; //The node that describes the object
+				UInterchangeBaseNode* FactoryNode = nullptr; //The node that describes the object
 				bool bIsReimport;
 			};
 
