Interchange scene import: Set mesh on static and skeletal mesh components
- Renamed Target Asset Uids to Target Node Uids and used it bidirectionaly.
- Target Node Uids is mutable so that a const node graph (translated nodes) can still track which factory nodes its being used by.
- Only register the actor components in the Interchange completion task so that we know that the assets are ready.
- Added transform compensation for FBX cameras and lights.

diff --git a/Engine/Plugins/Runtime/Interchange/Source/Import/Private/Scene/InterchangeActorFactory.cpp b/Engine/Plugins/Runtime/Interchange/Source/Import/Private/Scene/InterchangeActorFactory.cpp
--- a/Engine/Plugins/Runtime/Interchange/Source/Import/Private/Scene/InterchangeActorFactory.cpp
+++ b/Engine/Plugins/Runtime/Interchange/Source/Import/Private/Scene/InterchangeActorFactory.cpp
@@ -8,11 +8,14 @@
 
 #include "Components/SceneComponent.h"
 #include "Engine/Engine.h"
+#include "Engine/StaticMesh.h"
 #include "Engine/World.h"
 #include "GameFramework/Actor.h"
 
 #include "Nodes/InterchangeBaseNode.h"
 #include "Nodes/InterchangeBaseNodeContainer.h"
+#include "InterchangeMeshNode.h"
+#include "InterchangeSceneNode.h"
 
 #if WITH_EDITOR
 #include "Editor/EditorEngine.h"
@@ -24,6 +27,65 @@ namespace UE
 	{
 		namespace Private
 		{
+			FString FindMeshDependencyUid(const UInterchangeBaseNodeContainer* NodeContainer, const UInterchangeActorFactoryNode* ActorFactoryNode)
+			{
+				FString MeshDependency;
+
+				TArray<FString> ActorTargetNodes;
+				ActorFactoryNode->GetTargetNodeUids(ActorTargetNodes);
+				if (!ActorTargetNodes.IsEmpty())
+				{
+					if (const UInterchangeSceneNode* SceneNode = Cast<UInterchangeSceneNode>(NodeContainer->GetNode(ActorTargetNodes[0])))
+					{
+						FString AssetInstanceUid;
+						SceneNode->GetCustomAssetInstanceUid(AssetInstanceUid);
+
+						if (const UInterchangeMeshNode* MeshNode = Cast<UInterchangeMeshNode>(NodeContainer->GetNode(AssetInstanceUid)))
+						{
+							TArray<FString> MeshTargetNodes;
+							MeshNode->GetTargetNodeUids(MeshTargetNodes);
+
+							if (!MeshTargetNodes.IsEmpty())
+							{
+								MeshDependency = MeshTargetNodes[0];
+							}
+						}
+					}
+				}
+
+				return MeshDependency;
+			}
+
+			void SetupStaticMeshComponent(const UInterchangeBaseNodeContainer* NodeContainer, const UInterchangeActorFactoryNode* ActorFactoryNode, UStaticMeshComponent& StaticMeshComponent)
+			{
+				StaticMeshComponent.UnregisterComponent();
+
+				const FString MeshDependency = FindMeshDependencyUid(NodeContainer, ActorFactoryNode);
+
+				if (const UInterchangeBaseNode* MeshNode = NodeContainer->GetNode(MeshDependency))
+				{
+					if (UStaticMesh* StaticMesh = Cast<UStaticMesh>(MeshNode->ReferenceObject.TryLoad()))
+					{
+						StaticMeshComponent.SetStaticMesh(StaticMesh);
+					}
+				}
+			}
+
+			void SetupSkeletalMeshComponent(const UInterchangeBaseNodeContainer* NodeContainer, const UInterchangeActorFactoryNode* ActorFactoryNode, USkeletalMeshComponent& SkeletalMeshComponent)
+			{
+				SkeletalMeshComponent.UnregisterComponent();
+
+				const FString MeshDependency = FindMeshDependencyUid(NodeContainer, ActorFactoryNode);
+
+				if (const UInterchangeBaseNode* MeshNode = NodeContainer->GetNode(MeshDependency))
+				{
+					if (USkeletalMesh* SkeletalMesh = Cast<USkeletalMesh>(MeshNode->ReferenceObject.TryLoad()))
+					{
+						SkeletalMeshComponent.SetSkeletalMesh(SkeletalMesh);
+					}
+				}
+			}
+
 			AActor* SpawnActor(const FActorSpawnParameters& ActorSpawnParameters, const UInterchangeActorFactoryNode* ActorFactoryNode, AActor* ParentActor)
 			{
 				ensure(ActorFactoryNode);
@@ -110,6 +172,21 @@ namespace UE
 				const FString ActorNodeUniqueID = ActorNode->GetUniqueID();
 
 				AActor* SpawnedActor = SpawnActor(SpawnParameters, Cast<UInterchangeActorFactoryNode>(ActorNode), ParentActor);
+
+				if (!SpawnedActor)
+				{
+					return;
+				}
+
+				if (UStaticMeshComponent* StaticMeshComponent = Cast<UStaticMeshComponent>(SpawnedActor->GetRootComponent()))
+				{
+					SetupStaticMeshComponent(NodeContainer, ActorNode, *StaticMeshComponent);
+				}
+				else if (USkeletalMeshComponent* SkeletalMeshComponent = Cast<USkeletalMeshComponent>(SpawnedActor->GetRootComponent()))
+				{
+					SetupSkeletalMeshComponent(NodeContainer, ActorNode, *SkeletalMeshComponent);
+				}
+
 				OutSpawnedActors.Add(ActorNodeUniqueID, SpawnedActor);
 
 				if (bSpawnChildren)
@@ -156,4 +233,4 @@ TMap<FString, UObject*> UInterchangeActorFactory::CreateSceneObjects(const UInte
 		ActorSpawnParameters, ParentActor, CreateSceneObjectsParams.bCreateSceneObjectsForChildren);
 
 	return SpawnedActors;
-}
\ No newline at end of file
+}
diff --git a/Engine/Plugins/Runtime/Interchange/Source/Parsers/Fbx/Private/FbxConvert.cpp b/Engine/Plugins/Runtime/Interchange/Source/Parsers/Fbx/Private/FbxConvert.cpp
--- a/Engine/Plugins/Runtime/Interchange/Source/Parsers/Fbx/Private/FbxConvert.cpp
+++ b/Engine/Plugins/Runtime/Interchange/Source/Parsers/Fbx/Private/FbxConvert.cpp
@@ -83,6 +83,26 @@ namespace UE
 				return Out;
 			}
 
+			FTransform FFbxConvert::AdjustCameraTransform(const FTransform& Transform)
+			{
+				//Add a roll of -90 degree locally for every cameras. Camera up vector differ from fbx to unreal
+				const FRotator AdditionalRotation(0.0f, 0.0f, -90.0f);
+				FTransform CameraTransform = FTransform(AdditionalRotation) * Transform;
+
+				//Remove the scale of the node holding a camera (the mesh is provide by the engine and can be different in size)
+				CameraTransform.SetScale3D(FVector::OneVector);
+				
+				return CameraTransform;
+			}
+
+			FTransform FFbxConvert::AdjustLightTransform(const FTransform& Transform)
+			{
+				//Add the z rotation of 90 degree locally for every light. Light direction differ from fbx to unreal 
+				const FRotator AdditionalRotation(0.0f, 90.0f, 0.0f);
+				FTransform LightTransform = FTransform(AdditionalRotation) * Transform;
+				return LightTransform;
+			}
+
 			FMatrix FFbxConvert::ConvertMatrix(const FbxAMatrix& Matrix)
 			{
 				FMatrix UEMatrix;
diff --git a/Engine/Plugins/Runtime/Interchange/Source/Parsers/Fbx/Private/FbxConvert.h b/Engine/Plugins/Runtime/Interchange/Source/Parsers/Fbx/Private/FbxConvert.h
--- a/Engine/Plugins/Runtime/Interchange/Source/Parsers/Fbx/Private/FbxConvert.h
+++ b/Engine/Plugins/Runtime/Interchange/Source/Parsers/Fbx/Private/FbxConvert.h
@@ -30,6 +30,9 @@ namespace UE
 				static FVector ConvertPos(const FbxVector4& Vector);
 				static FVector ConvertDir(const FbxVector4& Vector);
 
+				static FTransform AdjustCameraTransform(const FTransform& Transform);
+				static FTransform AdjustLightTransform(const FTransform& Transform);
+
 				/** Transform Conversion API End */
 				//////////////////////////////////////////////////////////////////////////
 
diff --git a/Engine/Plugins/Runtime/Interchange/Source/Parsers/Fbx/Private/FbxScene.cpp b/Engine/Plugins/Runtime/Interchange/Source/Parsers/Fbx/Private/FbxScene.cpp
--- a/Engine/Plugins/Runtime/Interchange/Source/Parsers/Fbx/Private/FbxScene.cpp
+++ b/Engine/Plugins/Runtime/Interchange/Source/Parsers/Fbx/Private/FbxScene.cpp
@@ -76,8 +76,8 @@ namespace UE
 
 				UInterchangeSceneNode* UnrealNode = CreateTransformNode(NodeContainer, NodeName, NodeUniqueID);
 				check(UnrealNode);
-					
-				auto GetConvertedTransform = [](FbxAMatrix& NewFbxMatrix)
+				
+				auto GetConvertedTransform = [Node](FbxAMatrix& NewFbxMatrix)
 				{
 					FTransform Transform;
 					FbxVector4 NewLocalT = NewFbxMatrix.GetT();
@@ -86,13 +86,27 @@ namespace UE
 					Transform.SetTranslation(FFbxConvert::ConvertPos(NewLocalT));
 					Transform.SetScale3D(FFbxConvert::ConvertScale(NewLocalS));
 					Transform.SetRotation(FFbxConvert::ConvertRotToQuat(NewLocalQ));
+
+					if (FbxNodeAttribute* NodeAttribute = Node->GetNodeAttribute())
+					{
+						switch (NodeAttribute->GetAttributeType())
+						{
+						case FbxNodeAttribute::eCamera:
+							Transform = FFbxConvert::AdjustCameraTransform(Transform);
+							break;
+						case FbxNodeAttribute::eLight:
+							Transform = FFbxConvert::AdjustLightTransform(Transform);
+							break;
+						}
+					}
+
 					return Transform;
 				};
 
 				{
 					//Set the global node transform
-					FbxAMatrix GloabalFbxMatrix = Node->EvaluateGlobalTransform();
-					FTransform GlobalTransform = GetConvertedTransform(GloabalFbxMatrix);
+					FbxAMatrix GlobalFbxMatrix = Node->EvaluateGlobalTransform();
+					FTransform GlobalTransform = GetConvertedTransform(GlobalFbxMatrix);
 					UnrealNode->SetCustomGlobalTransform(GlobalTransform);
 
 					//Set the local node transform
diff --git a/Engine/Plugins/Runtime/Interchange/Source/Pipelines/Private/InterchangeGenericAssetsPipeline.cpp b/Engine/Plugins/Runtime/Interchange/Source/Pipelines/Private/InterchangeGenericAssetsPipeline.cpp
--- a/Engine/Plugins/Runtime/Interchange/Source/Pipelines/Private/InterchangeGenericAssetsPipeline.cpp
+++ b/Engine/Plugins/Runtime/Interchange/Source/Pipelines/Private/InterchangeGenericAssetsPipeline.cpp
@@ -113,7 +113,7 @@ bool UInterchangeGenericAssetsPipeline::ExecutePreImportPipeline(UInterchangeBas
 	//Find all translated node we need for this pipeline
 	BaseNodeContainer->IterateNodes([this](const FString& NodeUid, UInterchangeBaseNode* Node)
 	{
-		switch(Node->GetnodeContainerType())
+		switch(Node->GetNodeContainerType())
 		{
 			case EInterchangeNodeContainerType::NodeContainerType_TranslatedAsset:
 			{
@@ -384,7 +384,8 @@ UInterchangeMaterialFactoryNode* UInterchangeGenericAssetsPipeline::CreateMateri
 		}
 		BaseNodeContainer->AddNode(MaterialFactoryNode);
 		MaterialFactoryNodes.Add(MaterialFactoryNode);
-		MaterialFactoryNode->AddTargetAssetUid(MaterialNode->GetUniqueID());
+		MaterialFactoryNode->AddTargetNodeUid(MaterialNode->GetUniqueID());
+		MaterialNode->AddTargetNodeUid(MaterialFactoryNode->GetUniqueID());
 	}
 	return MaterialFactoryNode;
 }
@@ -585,14 +586,16 @@ void UInterchangeGenericAssetsPipeline::AddLodDataToStaticMesh(UInterchangeStati
 				SceneNode->GetCustomAssetInstanceUid(MeshDependency);
 				if (BaseNodeContainer->IsNodeUidValid(MeshDependency))
 				{
-					StaticMeshFactoryNode->AddTargetAssetUid(MeshDependency);
+					StaticMeshFactoryNode->AddTargetNodeUid(MeshDependency);
+					BaseNodeContainer->GetNode(MeshDependency)->AddTargetNodeUid(StaticMeshFactoryNode->GetUniqueID());
 				}
 				
 				SceneNode->GetMaterialDependencyUids(MaterialDependencies);
 			}
 			else if (const UInterchangeMeshNode* MeshNode = Cast<UInterchangeMeshNode>(BaseNodeContainer->GetNode(NodeUid)))
 			{
-				StaticMeshFactoryNode->AddTargetAssetUid(NodeUid);
+				StaticMeshFactoryNode->AddTargetNodeUid(NodeUid);
+				MeshNode->AddTargetNodeUid(StaticMeshFactoryNode->GetUniqueID());
 				MeshNode->GetMaterialDependencies(MaterialDependencies);
 			}
 
diff --git a/Engine/Plugins/Runtime/Interchange/Source/Pipelines/Private/InterchangeGenericScenesPipeline.cpp b/Engine/Plugins/Runtime/Interchange/Source/Pipelines/Private/InterchangeGenericScenesPipeline.cpp
--- a/Engine/Plugins/Runtime/Interchange/Source/Pipelines/Private/InterchangeGenericScenesPipeline.cpp
+++ b/Engine/Plugins/Runtime/Interchange/Source/Pipelines/Private/InterchangeGenericScenesPipeline.cpp
@@ -27,7 +27,7 @@ bool UInterchangeGenericLevelPipeline::ExecutePreImportPipeline(UInterchangeBase
 	//Find all translated node we need for this pipeline
 	InBaseNodeContainer->IterateNodes([&SceneNodes](const FString& NodeUid, UInterchangeBaseNode* Node)
 	{
-		switch(Node->GetnodeContainerType())
+		switch(Node->GetNodeContainerType())
 		{
 		case EInterchangeNodeContainerType::NodeContainerType_TranslatedScene:
 		{
@@ -73,6 +73,8 @@ void UInterchangeGenericLevelPipeline::CreateActorFactoryNode(const UInterchange
 		ActorFactoryNode->SetParentUid(TEXT("Factory_") + SceneNode->GetParentUid());
 	}
 
+	ActorFactoryNode->AddTargetNodeUid(SceneNode->GetUniqueID());
+
 	FTransform GlobalTransform;
 	if (SceneNode->GetCustomGlobalTransform(GlobalTransform))
 	{
diff --git a/Engine/Plugins/Runtime/Interchange/Source/Pipelines/Private/InterchangeGenericSkeletalMeshPipeline.cpp b/Engine/Plugins/Runtime/Interchange/Source/Pipelines/Private/InterchangeGenericSkeletalMeshPipeline.cpp
--- a/Engine/Plugins/Runtime/Interchange/Source/Pipelines/Private/InterchangeGenericSkeletalMeshPipeline.cpp
+++ b/Engine/Plugins/Runtime/Interchange/Source/Pipelines/Private/InterchangeGenericSkeletalMeshPipeline.cpp
@@ -220,7 +220,7 @@ bool UInterchangeGenericAssetsPipeline::ExecutePreImportPipelineSkeletalMesh()
 	//Find all translated node we need for this pipeline
 	BaseNodeContainer->IterateNodes([this, &SkinnedMeshNodes, &StaticMeshNodes](const FString& NodeUid, UInterchangeBaseNode* Node)
 	{
-		switch(Node->GetnodeContainerType())
+		switch(Node->GetNodeContainerType())
 		{
 			case EInterchangeNodeContainerType::NodeContainerType_TranslatedAsset:
 			{
@@ -679,10 +679,20 @@ void UInterchangeGenericAssetsPipeline::AddLodDataToSkeletalMesh(const UIntercha
 			TArray<FString> MaterialDependencies;
 			if (const UInterchangeSceneNode* SceneNode = Cast<UInterchangeSceneNode>(BaseNodeContainer->GetNode(NodeUid)))
 			{
+				FString MeshDependency;
+				SceneNode->GetCustomAssetInstanceUid(MeshDependency);
+				if (BaseNodeContainer->IsNodeUidValid(MeshDependency))
+				{
+					SkeletalMeshFactoryNode->AddTargetNodeUid(MeshDependency);
+					BaseNodeContainer->GetNode(MeshDependency)->AddTargetNodeUid(SkeletalMeshFactoryNode->GetUniqueID());
+				}
+
 				SceneNode->GetMaterialDependencyUids(MaterialDependencies);
 			}
 			else if (const UInterchangeMeshNode* MeshNode = Cast<UInterchangeMeshNode>(BaseNodeContainer->GetNode(NodeUid)))
 			{
+				SkeletalMeshFactoryNode->AddTargetNodeUid(NodeUid);
+				MeshNode->AddTargetNodeUid(SkeletalMeshFactoryNode->GetUniqueID());
 				MeshNode->GetMaterialDependencies(MaterialDependencies);
 			}
 			const int32 MaterialCount = MaterialDependencies.Num();
diff --git a/Engine/Plugins/Runtime/Interchange/Source/Pipelines/Private/InterchangePipelineMeshesUtilities.cpp b/Engine/Plugins/Runtime/Interchange/Source/Pipelines/Private/InterchangePipelineMeshesUtilities.cpp
--- a/Engine/Plugins/Runtime/Interchange/Source/Pipelines/Private/InterchangePipelineMeshesUtilities.cpp
+++ b/Engine/Plugins/Runtime/Interchange/Source/Pipelines/Private/InterchangePipelineMeshesUtilities.cpp
@@ -21,7 +21,7 @@ UInterchangePipelineMeshesUtilities* UInterchangePipelineMeshesUtilities::Create
 	//Find all translated node we need for this pipeline
 	BaseNodeContainer->IterateNodes([&PipelineMeshesUtilities, &BaseNodeContainer, &SkeletonRootNodeUids](const FString& NodeUid, UInterchangeBaseNode* Node)
 	{
-		if (Node->GetnodeContainerType() == EInterchangeNodeContainerType::NodeContainerType_TranslatedAsset)
+		if (Node->GetNodeContainerType() == EInterchangeNodeContainerType::NodeContainerType_TranslatedAsset)
 		{
 			if (UInterchangeMeshNode* MeshNode = Cast<UInterchangeMeshNode>(Node))
 			{
@@ -35,7 +35,7 @@ UInterchangePipelineMeshesUtilities* UInterchangePipelineMeshesUtilities::Create
 	//Find all translated scene node we need for this pipeline
 	BaseNodeContainer->IterateNodes([&PipelineMeshesUtilities, &BaseNodeContainer, &SkeletonRootNodeUids](const FString& NodeUid, UInterchangeBaseNode* Node)
 	{
-		if (Node->GetnodeContainerType() == EInterchangeNodeContainerType::NodeContainerType_TranslatedScene)
+		if (Node->GetNodeContainerType() == EInterchangeNodeContainerType::NodeContainerType_TranslatedScene)
 		{
 			if (UInterchangeSceneNode* SceneNode = Cast<UInterchangeSceneNode>(Node))
 			{
diff --git a/Engine/Source/Runtime/Interchange/Core/Private/Nodes/InterchangeBaseNode.cpp b/Engine/Source/Runtime/Interchange/Core/Private/Nodes/InterchangeBaseNode.cpp
--- a/Engine/Source/Runtime/Interchange/Core/Private/Nodes/InterchangeBaseNode.cpp
+++ b/Engine/Source/Runtime/Interchange/Core/Private/Nodes/InterchangeBaseNode.cpp
@@ -8,7 +8,14 @@
 #include "UObject/ObjectMacros.h"
 
 
-
+UInterchangeBaseNode::UInterchangeBaseNode()
+{
+	Attributes = MakeShared<UE::Interchange::FAttributeStorage, ESPMode::ThreadSafe>();
+	FactoryDependencies.Initialize(Attributes, UE::Interchange::FBaseNodeStaticData::GetFactoryDependenciesBaseKey());
+	TargetNodes.Initialize(Attributes, UE::Interchange::FBaseNodeStaticData::TargetAssetIDsKey());
+	RegisterAttribute<bool>(UE::Interchange::FBaseNodeStaticData::IsEnabledKey(), true);
+	RegisterAttribute<uint8>(UE::Interchange::FBaseNodeStaticData::NodeContainerTypeKey(), static_cast<uint8>(EInterchangeNodeContainerType::NodeContainerType_None));
+}
 
 void UInterchangeBaseNode::InitializeNode(const FString& UniqueID, const FString& DisplayLabel, const EInterchangeNodeContainerType NodeContainerType)
 {
@@ -30,8 +37,6 @@ void UInterchangeBaseNode::InitializeNode(const FString& UniqueID, const FString
 		LogAttributeStorageErrors(Result, TEXT("RegisterAttribute"), UE::Interchange::FBaseNodeStaticData::NodeContainerTypeKey());
 	}
 
-	FactoryDependencies.Initialize(Attributes, UE::Interchange::FBaseNodeStaticData::GetFactoryDependenciesBaseKey());
-
 	bIsInitialized = true;
 }
 
@@ -156,7 +161,7 @@ bool UInterchangeBaseNode::SetEnabled(const bool bIsEnabled)
 	return false;
 }
 
-EInterchangeNodeContainerType UInterchangeBaseNode::GetnodeContainerType() const
+EInterchangeNodeContainerType UInterchangeBaseNode::GetNodeContainerType() const
 {
 	if (!Attributes->ContainAttribute(UE::Interchange::FBaseNodeStaticData::NodeContainerTypeKey()))
 	{
@@ -211,24 +216,24 @@ bool UInterchangeBaseNode::SetAssetName(const FString& AssetName)
 	return false;
 }
 
-int32 UInterchangeBaseNode::GetTargetAssetCount() const
+int32 UInterchangeBaseNode::GetTargetNodeCount() const
 {
-	return TargetAssets.GetCount();
+	return TargetNodes.GetCount();
 }
 
-void UInterchangeBaseNode::GetTargetAssetUids(TArray<FString>& OutTargetAssets) const
+void UInterchangeBaseNode::GetTargetNodeUids(TArray<FString>& OutTargetAssets) const
 {
-	TargetAssets.GetNames(OutTargetAssets);
+	TargetNodes.GetNames(OutTargetAssets);
 }
 
-bool UInterchangeBaseNode::AddTargetAssetUid(const FString& AssetUid)
+bool UInterchangeBaseNode::AddTargetNodeUid(const FString& AssetUid) const
 {
-	return TargetAssets.AddName(AssetUid);
+	return TargetNodes.AddName(AssetUid);
 }
 
-bool UInterchangeBaseNode::RemoveTargetAssetUid(const FString& AssetUid)
+bool UInterchangeBaseNode::RemoveTargetNodeUid(const FString& AssetUid) const
 {
-	return TargetAssets.RemoveName(AssetUid);
+	return TargetNodes.RemoveName(AssetUid);
 }
 
 FString UInterchangeBaseNode::InvalidNodeUid()
diff --git a/Engine/Source/Runtime/Interchange/Core/Public/Nodes/InterchangeBaseNode.h b/Engine/Source/Runtime/Interchange/Core/Public/Nodes/InterchangeBaseNode.h
--- a/Engine/Source/Runtime/Interchange/Core/Public/Nodes/InterchangeBaseNode.h
+++ b/Engine/Source/Runtime/Interchange/Core/Public/Nodes/InterchangeBaseNode.h
@@ -816,14 +816,7 @@ class INTERCHANGECORE_API UInterchangeBaseNode : public UObject
 	GENERATED_BODY()
 
 public:
-	UInterchangeBaseNode()
-	{
-		Attributes = MakeShared<UE::Interchange::FAttributeStorage, ESPMode::ThreadSafe>();
-		FactoryDependencies.Initialize(Attributes, UE::Interchange::FBaseNodeStaticData::GetFactoryDependenciesBaseKey());
-		TargetAssets.Initialize(Attributes, UE::Interchange::FBaseNodeStaticData::TargetAssetIDsKey());
-		RegisterAttribute<bool>(UE::Interchange::FBaseNodeStaticData::IsEnabledKey(), true);
-		RegisterAttribute<uint8>(UE::Interchange::FBaseNodeStaticData::NodeContainerTypeKey(), static_cast<uint8>(EInterchangeNodeContainerType::NodeContainerType_None));
-	}
+	UInterchangeBaseNode();
 
 	virtual ~UInterchangeBaseNode() = default;
 
@@ -1131,25 +1124,25 @@ public:
 	 * Get number of target assets relating to this object.
 	 */
 	UFUNCTION(BlueprintCallable, Category = "Interchange | Node")
-	int32 GetTargetAssetCount() const;
+	int32 GetTargetNodeCount() const;
 
 	/**
 	 * Get target assets relating to this object.
 	 */
 	UFUNCTION(BlueprintCallable, Category = "Interchange | Node")
-	void GetTargetAssetUids(TArray<FString>& OutTargetAssets) const;
+	void GetTargetNodeUids(TArray<FString>& OutTargetAssets) const;
 
 	/**
 	 * Add asset node UID relating to this object.
 	 */
 	UFUNCTION(BlueprintCallable, Category = "Interchange | Node")
-	bool AddTargetAssetUid(const FString& AssetUid);
+	bool AddTargetNodeUid(const FString& AssetUid) const;
 
 	/**
 	 * Remove asset node UID relating to this object.
 	 */
 	UFUNCTION(BlueprintCallable, Category = "Interchange | Node")
-	bool RemoveTargetAssetUid(const FString& AssetUid);
+	bool RemoveTargetNodeUid(const FString& AssetUid) const;
 
 	/**
 	 * IsEnable true mean that the node will be import/export, if false it will be discarded.
@@ -1170,7 +1163,7 @@ public:
 	 * Return the node container type which define the purpose of the node (Factory node, translated scene node or translated asset node).
 	 */
 	UFUNCTION(BlueprintCallable, Category = "Interchange | Node")
-	EInterchangeNodeContainerType GetnodeContainerType() const;
+	EInterchangeNodeContainerType GetNodeContainerType() const;
 
 	/**
 	 * Return a FGuid build from the FSHA1 of all the attribute data contain in the node.
@@ -1252,5 +1245,5 @@ protected:
 	/**
 	 * This tracks the IDs of asset nodes which are the target of factories
 	 */
-	UE::Interchange::FNameAttributeArrayHelper TargetAssets;
+	mutable UE::Interchange::FNameAttributeArrayHelper TargetNodes;
 };
diff --git a/Engine/Source/Runtime/Interchange/Engine/Private/Tasks/InterchangeTaskCompletion.cpp b/Engine/Source/Runtime/Interchange/Engine/Private/Tasks/InterchangeTaskCompletion.cpp
--- a/Engine/Source/Runtime/Interchange/Engine/Private/Tasks/InterchangeTaskCompletion.cpp
+++ b/Engine/Source/Runtime/Interchange/Engine/Private/Tasks/InterchangeTaskCompletion.cpp
@@ -99,6 +99,15 @@ void UE::Interchange::FTaskPreCompletion::DoTask(ENamedThreads::Type CurrentThre
 				}
 				else
 				{
+					if (AActor* Actor = Cast<AActor>(ImportedObject))
+					{
+						Actor->RegisterAllComponents();
+					}
+					else if (UActorComponent* Component = Cast<UActorComponent>(ImportedObject))
+					{
+						Component->RegisterComponent();
+					}
+
 					AsyncHelper->SceneImportResult->AddImportedObject(ImportedObject);
 				}
 
diff --git a/Engine/Source/Runtime/Interchange/Engine/Private/Tasks/InterchangeTaskCreateAsset.cpp b/Engine/Source/Runtime/Interchange/Engine/Private/Tasks/InterchangeTaskCreateAsset.cpp
--- a/Engine/Source/Runtime/Interchange/Engine/Private/Tasks/InterchangeTaskCreateAsset.cpp
+++ b/Engine/Source/Runtime/Interchange/Engine/Private/Tasks/InterchangeTaskCreateAsset.cpp
@@ -296,7 +296,7 @@ void UE::Interchange::FTaskCreateAsset::DoTask(ENamedThreads::Type CurrentThread
 				if (!Result->InterchangeKey.IsEmpty() && (Result->DestinationAssetName.IsEmpty() || Result->AssetType == nullptr))
 				{
 					TArray<FString> TargetAssets;
-					Node->GetTargetAssetUids(TargetAssets);
+					Node->GetTargetNodeUids(TargetAssets);
 					if (TargetAssets.Contains(Result->InterchangeKey))
 					{
 						Result->DestinationAssetName = NodeAsset->GetPathName();
