Interchange support properly multiple source files when creating or reimporting an asset.

diff --git a/Engine/Plugins/Runtime/Interchange/Source/Import/Private/InterchangeImportCommon.cpp b/Engine/Plugins/Runtime/Interchange/Source/Import/Private/InterchangeImportCommon.cpp
--- a/Engine/Plugins/Runtime/Interchange/Source/Import/Private/InterchangeImportCommon.cpp
+++ b/Engine/Plugins/Runtime/Interchange/Source/Import/Private/InterchangeImportCommon.cpp
@@ -8,11 +8,58 @@
 #include "InterchangeSourceData.h"
 #include "Nodes/InterchangeBaseNodeContainer.h"
 #include "UObject/Object.h"
+#include "EditorFramework/AssetImportData.h"
+#include "HAL/FileManager.h"
 
 namespace UE
 {
 	namespace Interchange
 	{
+		namespace Private::ImportCommon
+		{
+			UInterchangeAssetImportData* BeginSetupAssetData(FFactoryCommon::FUpdateImportAssetDataParameters& Parameters)
+			{
+				if (!ensure(IsInGameThread()))
+				{
+					return nullptr;
+				}
+				if (!ensure(Parameters.SourceData && Parameters.AssetImportDataOuter))
+				{
+					return nullptr;
+				}
+
+				UInterchangeAssetImportData* AssetImportData = Cast<UInterchangeAssetImportData>(Parameters.AssetImportData);
+
+				if (!AssetImportData)
+				{
+					AssetImportData = NewObject<UInterchangeAssetImportData>(Parameters.AssetImportDataOuter, NAME_None);
+				}
+
+				ensure(AssetImportData);
+
+				return AssetImportData;
+			}
+
+
+			void EndSetupAssetData(FFactoryCommon::FUpdateImportAssetDataParameters& Parameters, UInterchangeAssetImportData* AssetImportData)
+			{
+				//Set the interchange node graph data
+				AssetImportData->NodeUniqueID = Parameters.NodeUniqueID;
+				FObjectDuplicationParameters DupParam(Parameters.NodeContainer, AssetImportData);
+				AssetImportData->NodeContainer = CastChecked<UInterchangeBaseNodeContainer>(StaticDuplicateObjectEx(DupParam));
+				for (const UInterchangePipelineBase* Pipeline : Parameters.Pipelines)
+				{
+					UInterchangePipelineBase* DupPipeline = Cast<UInterchangePipelineBase>(StaticDuplicateObject(Pipeline, AssetImportData));
+					if (DupPipeline)
+					{
+						AssetImportData->Pipelines.Add(DupPipeline);
+					}
+				}
+			}
+
+
+		}
+
 		FFactoryCommon::FUpdateImportAssetDataParameters::FUpdateImportAssetDataParameters(UObject* InAssetImportDataOuter
 																							, UAssetImportData* InAssetImportData
 																							, const UInterchangeSourceData* InSourceData
@@ -43,62 +90,89 @@ namespace UE
 			{
 				return nullptr;
 			}
-			//Set the asset import data file source to allow reimport. TODO: manage MD5 Hash properly
-			TOptional<FMD5Hash> FileContentHash = Parameters.SourceData->GetFileContentHash();
 
-			UInterchangeAssetImportData* AssetImportData = nullptr;
-				
+			UInterchangeAssetImportData* AssetImportData = Private::ImportCommon::BeginSetupAssetData(Parameters);
+
 			if (Parameters.AssetImportData)
 			{
-				AssetImportData = Cast<UInterchangeAssetImportData>(Parameters.AssetImportData);
-				if (!AssetImportData)
+				if (!Parameters.AssetImportData->IsA<UInterchangeAssetImportData>())
 				{
-					AssetImportData = NewObject<UInterchangeAssetImportData>(Parameters.AssetImportDataOuter, NAME_None);
-					TArray<FString> Filenames;
-					Parameters.AssetImportData->ExtractFilenames(Filenames);
-					TArray<FString> Labels;
-					Parameters.AssetImportData->ExtractDisplayLabels(Labels);
-					if (Filenames.Num() > 1)
-					{
-						for (int32 FileIndex = 0; FileIndex < Filenames.Num(); ++FileIndex)
-						{
-							FString Filename = Filenames[FileIndex];
-							FString Label;
-							if (Labels.IsValidIndex(FileIndex))
-							{
-								Label = Labels[FileIndex];
-							}
-							//This is slow since it will hash the file, TODO make sure hashing is done in the create asset task
-							AssetImportData->AddFileName(Filename, FileIndex, Label);
-						}
-					}
+					// Migrate the old source data
+					TArray<FAssetImportInfo::FSourceFile> OldSourceFiles = Parameters.AssetImportData->SourceData.SourceFiles;
+					AssetImportData->SetSourceFiles(MoveTemp(OldSourceFiles));
 				}
 			}
-			else
-			{
-				AssetImportData = NewObject<UInterchangeAssetImportData>(Parameters.AssetImportDataOuter, NAME_None);
-			}
 
-			ensure(AssetImportData);
+			//Set the asset import data file source to allow reimport. TODO: manage MD5 Hash properly
+			TOptional<FMD5Hash> FileContentHash = Parameters.SourceData->GetFileContentHash();
+
 			//Update the first filename, TODO for asset using multiple source file we have to update the correct index
 			AssetImportData->Update(Parameters.SourceData->GetFilename(), FileContentHash.IsSet() ? &FileContentHash.GetValue() : nullptr);
-			//Set the interchange node graph data
-			AssetImportData->NodeUniqueID = Parameters.NodeUniqueID;
-			FObjectDuplicationParameters DupParam(Parameters.NodeContainer, AssetImportData);
-			AssetImportData->NodeContainer = CastChecked<UInterchangeBaseNodeContainer>(StaticDuplicateObjectEx(DupParam));
-			for (const UInterchangePipelineBase* Pipeline : Parameters.Pipelines)
+
+
+			Private::ImportCommon::EndSetupAssetData(Parameters, AssetImportData);
+
+			// Return the asset import data so it can be set on the imported asset.
+			return AssetImportData;
+#endif //#if WITH_EDITORONLY_DATA
+			return nullptr;
+		}
+
+
+#if WITH_EDITORONLY_DATA
+		FFactoryCommon::FSetImportAssetDataParameters::FSetImportAssetDataParameters(UObject* InAssetImportDataOuter
+																					, UAssetImportData* InAssetImportData
+																					, const UInterchangeSourceData* InSourceData
+																					, FString InNodeUniqueID
+																					, UInterchangeBaseNodeContainer* InNodeContainer
+																					, const TArray<UInterchangePipelineBase*>& InPipelines)
+			: FUpdateImportAssetDataParameters(InAssetImportDataOuter
+				, InAssetImportData
+				, InSourceData
+				, InNodeUniqueID
+				, InNodeContainer
+				, InPipelines
+				)
+			, SourceFiles()
+		{
+		}
+
+		UAssetImportData* FFactoryCommon::SetImportAssetData(FSetImportAssetDataParameters& Parameters)
+		{
+			UInterchangeAssetImportData* AssetImportData = Private::ImportCommon::BeginSetupAssetData(Parameters);
+
+			// Update the source files
 			{
-				UInterchangePipelineBase* DupPipeline = Cast<UInterchangePipelineBase>(StaticDuplicateObject(Pipeline, AssetImportData));
-				if (DupPipeline)
+				if (Parameters.SourceFiles.IsEmpty())
 				{
-					AssetImportData->Pipelines.Add(DupPipeline);
+					TOptional<FMD5Hash> FileContentHash = Parameters.SourceData->GetFileContentHash();
+
+					// Todo add display label?
+					Parameters.SourceFiles.Emplace(
+						AssetImportData->SanitizeImportFilename(Parameters.SourceData->GetFilename()),
+						IFileManager::Get().GetTimeStamp(*Parameters.SourceData->GetFilename()),
+						FileContentHash.IsSet() ? FileContentHash.GetValue() : FMD5Hash()
+						);
 				}
+				else
+				{
+					for (FAssetImportInfo::FSourceFile& Source : Parameters.SourceFiles)
+					{
+						// This is done here since this is not thread safe
+						Source.RelativeFilename = AssetImportData->SanitizeImportFilename(Source.RelativeFilename);
+					}
+				}
+
+				AssetImportData->SetSourceFiles(MoveTemp(Parameters.SourceFiles));
 			}
+
+			Private::ImportCommon::EndSetupAssetData(Parameters, AssetImportData);
+
 			// Return the asset import data so it can be set on the imported asset.
 			return AssetImportData;
-#endif //#if WITH_EDITORONLY_DATA
-			return nullptr;
 		}
+#endif //#if WITH_EDITORONLY_DATA
+
 
 		void FFactoryCommon::ApplyReimportStrategyToAsset(const EReimportStrategyFlags ReimportStrategyFlags
 										  , UObject* Asset
diff --git a/Engine/Plugins/Runtime/Interchange/Source/Import/Private/InterchangeImportCommon.h b/Engine/Plugins/Runtime/Interchange/Source/Import/Private/InterchangeImportCommon.h
--- a/Engine/Plugins/Runtime/Interchange/Source/Import/Private/InterchangeImportCommon.h
+++ b/Engine/Plugins/Runtime/Interchange/Source/Import/Private/InterchangeImportCommon.h
@@ -2,6 +2,7 @@
 #pragma once
 
 #include "CoreMinimal.h"
+#include "EditorFramework/AssetImportData.h"
 #include "InterchangeFactoryBase.h"
 
 class UAssetImportData;
@@ -47,6 +48,28 @@ namespace UE
 			 */
 			static UAssetImportData* UpdateImportAssetData(FUpdateImportAssetDataParameters& Parameters);
 
+#if WITH_EDITORONLY_DATA
+			struct FSetImportAssetDataParameters : public FUpdateImportAssetDataParameters
+			{
+				// Allow the factory to provide is own list of source files.
+				TArray<FAssetImportInfo::FSourceFile> SourceFiles;
+
+				FSetImportAssetDataParameters(UObject* InAssetImportDataOuter
+					, UAssetImportData* InAssetImportData
+					, const UInterchangeSourceData* InSourceData
+					, FString InNodeUniqueID
+					, UInterchangeBaseNodeContainer* InNodeContainer
+					, const TArray<UInterchangePipelineBase*>& InPipelines);
+			};
+
+			/**
+			 * Set the AssetImportData source file of the specified asset in the parameters. Also update the node container and the node unique id.
+			 * If the AssetImportData is null it will create one. If the AssetImportData is not an UInterchangeAssetImportData it will create a new one.
+			 * @return The source data that should be stored on the asset or nullptr if a parameter is invalid
+			 */
+			static UAssetImportData* SetImportAssetData(FSetImportAssetDataParameters& Parameters);
+#endif // WITH_EDITORONLY_DATA
+
 			/**
 			 * Apply the current strategy to the PipelineAssetNode
 			 */
diff --git a/Engine/Plugins/Runtime/Interchange/Source/Import/Private/Texture/InterchangeTextureFactory.cpp b/Engine/Plugins/Runtime/Interchange/Source/Import/Private/Texture/InterchangeTextureFactory.cpp
--- a/Engine/Plugins/Runtime/Interchange/Source/Import/Private/Texture/InterchangeTextureFactory.cpp
+++ b/Engine/Plugins/Runtime/Interchange/Source/Import/Private/Texture/InterchangeTextureFactory.cpp
@@ -2,11 +2,14 @@
 
 #include "Texture/InterchangeTextureFactory.h"
 
+#include "Async/TaskGraphInterfaces.h"
+#include "Async/TaskGraphInterfaces.h"
 #include "Engine/Texture.h"
 #include "Engine/Texture2D.h"
 #include "Engine/Texture2DArray.h"
 #include "Engine/TextureCube.h"
 #include "Engine/TextureLightProfile.h"
+#include "HAL/FileManagerGeneric.h"
 #include "InterchangeAssetImportData.h"
 #include "InterchangeImportCommon.h"
 #include "InterchangeImportLog.h"
@@ -27,6 +30,7 @@
 #include "Texture/InterchangeTextureLightProfilePayloadInterface.h"
 #include "Texture/InterchangeTexturePayloadInterface.h"
 #include "TextureCompiler.h"
+#include "UDIMUtilities.h"
 #include "UObject/ObjectMacros.h"
 
 #if WITH_EDITORONLY_DATA
@@ -36,7 +40,6 @@
 #endif //WITH_EDITORONLY_DATA
 
 
-
 namespace UE::Interchange::Private::InterchangeTextureFactory
 {
 	/**
@@ -585,6 +588,139 @@ namespace UE::Interchange::Private::InterchangeTextureFactory
 		return *this;
 	}
 
+	TArray<FString> GetFilesToHash(const FTextureNodeVariant& TextureNodeVariant, const FTexturePayloadVariant& TexturePayload)
+	{
+		TArray<FString> FilesToHash;
+		// Standard texture 2D payload
+		if (const UInterchangeTexture2DNode* const* TextureNode = TextureNodeVariant.TryGet<const UInterchangeTexture2DNode*>())
+		{
+			using namespace UE::Interchange;
+			if (const TOptional<FImportBlockedImage>* OptionalBlockedPayload = TexturePayload.TryGet<TOptional<FImportBlockedImage>>())
+			{
+				if (OptionalBlockedPayload->IsSet())
+				{
+					const FImportBlockedImage& BlockImage = OptionalBlockedPayload->GetValue();
+					TMap<int32, FString> BlockAndFiles = (*TextureNode)->GetSourceBlocks();
+					FilesToHash.Reserve(BlockAndFiles.Num());
+					for (const FTextureSourceBlock& BlockData : BlockImage.BlocksData)
+					{
+						if (FString* FilePath = BlockAndFiles.Find(UE::TextureUtilitiesCommon::GetUDIMIndex(BlockData.BlockX, BlockData.BlockY)))
+						{
+							FilesToHash.Add(*FilePath);
+						}
+					}
+				}
+
+			}
+		}
+		return FilesToHash;
+	}
+
+	FGraphEventArray GenerateHashSourceFilesTasks(const UInterchangeSourceData* SourceData, TArray<FString>&& FilesToHash, TArray<FAssetImportInfo::FSourceFile>& OutSourceFiles)
+	{
+		struct FHashSourceTaskBase
+		{
+			/**
+				* Returns the name of the thread that this task should run on.
+				*
+				* @return Always run on any thread.
+				*/
+			ENamedThreads::Type GetDesiredThread()
+			{
+				return ENamedThreads::AnyBackgroundThreadNormalTask;
+			}
+
+			/**
+				* Gets the task's stats tracking identifier.
+				*
+				* @return Stats identifier.
+				*/
+			TStatId GetStatId() const
+			{
+				return GET_STATID(STAT_TaskGraph_OtherTasks);
+			}
+
+			/**
+				* Gets the mode for tracking subsequent tasks.
+				*
+				* @return Always track subsequent tasks.
+				*/
+			static ESubsequentsMode::Type GetSubsequentsMode()
+			{
+				return ESubsequentsMode::TrackSubsequents;
+			}
+		};
+
+		FGraphEventArray TasksToDo;
+
+		// We do the hashing of the source files after the import to avoid a bigger memory overhead.
+		if (FilesToHash.IsEmpty())
+		{
+			struct FHashSingleSource : public FHashSourceTaskBase
+			{
+				FHashSingleSource(const UInterchangeSourceData* InSourceData)
+					: SourceData(InSourceData)
+				{}
+
+				/**
+				 * Performs the actual task.
+				 *
+				 * @param CurrentThread The thread that this task is executing on.
+				 * @param MyCompletionGraphEvent The completion event.
+				 */
+				void DoTask(ENamedThreads::Type CurrentThread, const FGraphEventRef& MyCompletionGraphEvent)
+				{
+					if (SourceData)
+					{ 
+						//Getting the file Hash will cache it into the source data
+						SourceData->GetFileContentHash();
+					}
+				}
+
+			private:
+				const UInterchangeSourceData* SourceData = nullptr;
+			};
+		
+			TasksToDo.Add(TGraphTask<FHashSingleSource>::CreateTask().ConstructAndDispatchWhenReady(SourceData));
+		}
+		else
+		{
+			struct FHashMutipleSource : public FHashSourceTaskBase
+			{
+				FHashMutipleSource(FString&& InFileToHash, FAssetImportInfo::FSourceFile& OutSourceFile)
+					: FileToHash(MoveTemp(InFileToHash))
+					, SourceFile(OutSourceFile)
+				{}
+
+				/**
+				 * Performs the actual task.
+				 *
+				 * @param CurrentThread The thread that this task is executing on.
+				 * @param MyCompletionGraphEvent The completion event.
+				 */
+				void DoTask(ENamedThreads::Type CurrentThread, const FGraphEventRef& MyCompletionGraphEvent)
+				{
+					SourceFile.FileHash = FMD5Hash::HashFile(*FileToHash);
+					SourceFile.Timestamp =IFileManager::Get().GetTimeStamp(*FileToHash);
+					SourceFile.RelativeFilename = MoveTemp(FileToHash);
+				}
+
+			private:
+				FString FileToHash;
+				FAssetImportInfo::FSourceFile& SourceFile;
+			};
+
+			OutSourceFiles.AddDefaulted(FilesToHash.Num());
+
+			for (int32 Index = 0; Index < FilesToHash.Num(); ++Index)
+			{
+				TasksToDo.Add(TGraphTask<FHashMutipleSource>::CreateTask().ConstructAndDispatchWhenReady(MoveTemp(FilesToHash[Index]), OutSourceFiles[Index]));
+			}
+		}
+
+		return TasksToDo;
+	}
+
 #endif // WITH_EDITORONLY_DATA
  }
 
@@ -773,12 +909,17 @@ UObject* UInterchangeTextureFactory::CreateAsset(const FCreateAssetParams& Argum
 		return nullptr;
 	}
 
-	//Getting the file Hash will cache it into the source data
-	Arguments.SourceData->GetFileContentHash();
+	FGraphEventArray TasksToDo = GenerateHashSourceFilesTasks(Arguments.SourceData, GetFilesToHash(TextureNodeVariant, TexturePayload), SourceFiles);
+
+	// Hash the payload while we hash the source files
 
 	// This will hash the payload to generate a unique ID before passing it to the virtualized bulkdata
 	ProcessedPayload = MoveTemp(TexturePayload);
 
+	// Wait for the hashing task(s)
+	ENamedThreads::Type NamedThread = IsInGameThread() ? ENamedThreads::GameThread : ENamedThreads::AnyThread;
+	FTaskGraphInterface::Get().WaitUntilTasksComplete(TasksToDo, NamedThread);
+
 	//The interchange completion task (call in the GameThread after the factories pass), will call PostEditChange which will trig another asynchronous system that will build all texture in parallel
 
 	return Texture;
@@ -854,13 +995,15 @@ void UInterchangeTextureFactory::PreImportPreCompletedCallback(const FImportPreC
 	if (ensure(Texture && Arguments.SourceData))
 	{
 		//We must call the Update of the asset source file in the main thread because UAssetImportData::Update execute some delegate we do not control
-		UE::Interchange::FFactoryCommon::FUpdateImportAssetDataParameters UpdateImportAssetDataParameters(Texture
+		UE::Interchange::FFactoryCommon::FSetImportAssetDataParameters SetImportAssetDataParameters(Texture
 																						 , Texture->AssetImportData
 																						 , Arguments.SourceData
 																						 , Arguments.NodeUniqueID
 																						 , Arguments.NodeContainer
 																						 , Arguments.Pipelines);
-		Texture->AssetImportData = UE::Interchange::FFactoryCommon::UpdateImportAssetData(UpdateImportAssetDataParameters);
+		SetImportAssetDataParameters.SourceFiles = MoveTemp(SourceFiles);
+
+		Texture->AssetImportData = UE::Interchange::FFactoryCommon::SetImportAssetData(SetImportAssetDataParameters);
 	}
 #endif
 }
diff --git a/Engine/Plugins/Runtime/Interchange/Source/Import/Private/Texture/InterchangeUDIMTranslator.cpp b/Engine/Plugins/Runtime/Interchange/Source/Import/Private/Texture/InterchangeUDIMTranslator.cpp
--- a/Engine/Plugins/Runtime/Interchange/Source/Import/Private/Texture/InterchangeUDIMTranslator.cpp
+++ b/Engine/Plugins/Runtime/Interchange/Source/Import/Private/Texture/InterchangeUDIMTranslator.cpp
@@ -328,9 +328,11 @@ TOptional<UE::Interchange::FImportBlockedImage> UInterchangeUDIMTranslator::GetB
 					FString PreUDIMName;
 					FString PostUDIMName;
 					const int32 UDIMIndex = UE::TextureUtilitiesCommon::ParseUDIMName(FilenameNoExtension, UdimRegexPattern, PreUDIMName, PostUDIMName);
-					
-					const int32 BlockX = (UDIMIndex - 1001) % 10;
-					const int32 BlockY = (UDIMIndex - 1001) / 10;
+
+					int32 BlockX;
+					int32 BlockY;
+
+					UE::TextureUtilitiesCommon::ExtractUDIMCoordinates(UDIMIndex, BlockX, BlockY);
 
 					BlockedImage.InitBlockFromImage(BlockX, BlockY, Image);
 
diff --git a/Engine/Plugins/Runtime/Interchange/Source/Import/Public/Texture/InterchangeTextureFactory.h b/Engine/Plugins/Runtime/Interchange/Source/Import/Public/Texture/InterchangeTextureFactory.h
--- a/Engine/Plugins/Runtime/Interchange/Source/Import/Public/Texture/InterchangeTextureFactory.h
+++ b/Engine/Plugins/Runtime/Interchange/Source/Import/Public/Texture/InterchangeTextureFactory.h
@@ -3,6 +3,7 @@
 #pragma once
 
 #include "CoreMinimal.h"
+#include "EditorFramework/AssetImportData.h"
 #include "InterchangeFactoryBase.h"
 #include "Misc/TVariant.h"
 #include "Texture/InterchangeBlockedTexturePayloadData.h"
@@ -58,6 +59,11 @@ public:
 
 private:
 	UE::Interchange::Private::InterchangeTextureFactory::FProcessedPayload ProcessedPayload;
+
+#if WITH_EDITORONLY_DATA
+	// When importing a UDIM the data for the source files will be stored here
+	TArray<FAssetImportInfo::FSourceFile> SourceFiles;
+#endif // WITH_EDITORONLY_DATA
 };
 
 
diff --git a/Engine/Source/Editor/UnrealEd/Private/Factories/EditorFactories.cpp b/Engine/Source/Editor/UnrealEd/Private/Factories/EditorFactories.cpp
--- a/Engine/Source/Editor/UnrealEd/Private/Factories/EditorFactories.cpp
+++ b/Engine/Source/Editor/UnrealEd/Private/Factories/EditorFactories.cpp
@@ -3590,8 +3590,7 @@ UTexture* UTextureFactory::ImportTextureUDIM(UClass* Class, UObject* InParent, F
 				{
 					const int32 UDIMIndex = It.Key;
 					FTextureSourceBlock* Block = new(SourceBlocks) FTextureSourceBlock();
-					Block->BlockX = (UDIMIndex - 1001) % 10;
-					Block->BlockY = (UDIMIndex - 1001) / 10;
+					UE::TextureUtilitiesCommon::ExtractUDIMCoordinates(UDIMIndex, Block->BlockX, Block->BlockY);
 					Block->SizeX = Image.SizeX;
 					Block->SizeY = Image.SizeY;
 					Block->NumSlices = 1;
diff --git a/Engine/Source/Runtime/Engine/Classes/EditorFramework/AssetImportData.h b/Engine/Source/Runtime/Engine/Classes/EditorFramework/AssetImportData.h
--- a/Engine/Source/Runtime/Engine/Classes/EditorFramework/AssetImportData.h
+++ b/Engine/Source/Runtime/Engine/Classes/EditorFramework/AssetImportData.h
@@ -26,6 +26,8 @@ struct ENGINE_API FAssetImportInfo
 
 	struct FSourceFile
 	{
+		FSourceFile() = default;
+
 		FSourceFile(FString InRelativeFilename, const FDateTime& InTimestamp = 0, const FMD5Hash& InFileHash = FMD5Hash(), const FString& InDisplayLabelName = FString())
 			: RelativeFilename(MoveTemp(InRelativeFilename))
 			, Timestamp(InTimestamp)
@@ -123,6 +125,9 @@ public:
 	/** Add a filename at the specific index. Will update the imported timespan and MD5. It will also add in between with empty filenames*/
 	void AddFileName(const FString& InPath, int32 Index, FString SourceFileLabel = FString());
 
+	/** Replace the source files with the one provided. The MD5 hashes will be computed if they aren't set */
+	void SetSourceFiles(TArray<FAssetImportInfo::FSourceFile>&& SourceFiles);
+
 #if WITH_EDITOR
 	/** If your asset import data flavor need to add some asset registry tag, override this function. */
 	virtual void AppendAssetRegistryTags(TArray<FAssetRegistryTag>& OutTags){}
@@ -157,11 +162,11 @@ public:
 
 	virtual void PostLoad() override;
 
-protected:
-
 	/** Convert an absolute import path so that it's relative to either this object's package, BaseDir() or leave it absolute */
 	FString SanitizeImportFilename(const FString& InPath) const;
 
+protected:
+
 	/** Resolve a filename that is relative to either this object's package, BaseDir() or absolute */
 	FString ResolveImportFilename(const FString& InRelativePath) const;
 
diff --git a/Engine/Source/Runtime/Engine/Private/EditorFramework/AssetImportData.cpp b/Engine/Source/Runtime/Engine/Private/EditorFramework/AssetImportData.cpp
--- a/Engine/Source/Runtime/Engine/Private/EditorFramework/AssetImportData.cpp
+++ b/Engine/Source/Runtime/Engine/Private/EditorFramework/AssetImportData.cpp
@@ -173,6 +173,28 @@ void UAssetImportData::AddFileName(const FString& InPath, int32 Index, FString S
 	OnImportDataChanged.Broadcast(Old, this);
 }
 
+void UAssetImportData::SetSourceFiles(TArray<FAssetImportInfo::FSourceFile>&& SourceFiles)
+{
+	FAssetImportInfo Old = SourceData;
+
+	for (FAssetImportInfo::FSourceFile& SourceFile : SourceFiles)
+	{
+		if (!SourceFile.FileHash.IsValid())
+		{
+			SourceFile.FileHash = FMD5Hash::HashFile(*SourceFile.RelativeFilename);
+		}
+
+		if (SourceFile.Timestamp == FDateTime())
+		{
+			SourceFile.Timestamp = IFileManager::Get().GetTimeStamp(*SourceFile.RelativeFilename);
+		}
+	}
+
+	SourceData.SourceFiles = MoveTemp(SourceFiles);
+
+	OnImportDataChanged.Broadcast(Old, this);
+}
+
 void UAssetImportData::Update(const FString& InPath, FMD5Hash *Md5Hash/* = nullptr*/)
 {
 	FAssetImportInfo Old = SourceData;
diff --git a/Engine/Source/Runtime/TextureUtilitiesCommon/Private/UDIMUtilities.cpp b/Engine/Source/Runtime/TextureUtilitiesCommon/Private/UDIMUtilities.cpp
--- a/Engine/Source/Runtime/TextureUtilitiesCommon/Private/UDIMUtilities.cpp
+++ b/Engine/Source/Runtime/TextureUtilitiesCommon/Private/UDIMUtilities.cpp
@@ -46,5 +46,16 @@ namespace UE
 
 			return UdimValue;
 		}
+
+		int32 GetUDIMIndex(int32 BlockX, int32 BlockY)
+		{
+			return BlockY * 10 + BlockX + 1001;
+		}
+
+		void ExtractUDIMCoordinates(int32 UDIMIndex, int32& OutBlockX, int32& OutBlockY)
+		{
+			OutBlockX = (UDIMIndex - 1001) % 10;
+			OutBlockY = (UDIMIndex - 1001) / 10;
+		}
 	}
 }
diff --git a/Engine/Source/Runtime/TextureUtilitiesCommon/Public/UDIMUtilities.h b/Engine/Source/Runtime/TextureUtilitiesCommon/Public/UDIMUtilities.h
--- a/Engine/Source/Runtime/TextureUtilitiesCommon/Public/UDIMUtilities.h
+++ b/Engine/Source/Runtime/TextureUtilitiesCommon/Public/UDIMUtilities.h
@@ -11,5 +11,9 @@ namespace UE
 	namespace TextureUtilitiesCommon
 	{
 		TEXTUREUTILITIESCOMMON_API uint32 ParseUDIMName(const FString& Name, const FString& UdimRegexPattern, FString& OutPrefixName, FString& OutPostfixName);
+
+		TEXTUREUTILITIESCOMMON_API int32 GetUDIMIndex(int32 BlockX, int32 BlockY);
+
+		TEXTUREUTILITIESCOMMON_API void ExtractUDIMCoordinates(int32 UDIMIndex, int32& OutBlockX, int32& OutBlockY);
 	}
 }
