Fixed an issue where the parameter was validated against himself instead of the member UInterchangeGenericAssetsPipeline.

diff --git a/Engine/Plugins/Runtime/Interchange/Source/Pipelines/Private/InterchangeGenericAssetsPipeline.cpp b/Engine/Plugins/Runtime/Interchange/Source/Pipelines/Private/InterchangeGenericAssetsPipeline.cpp
--- a/Engine/Plugins/Runtime/Interchange/Source/Pipelines/Private/InterchangeGenericAssetsPipeline.cpp
+++ b/Engine/Plugins/Runtime/Interchange/Source/Pipelines/Private/InterchangeGenericAssetsPipeline.cpp
@@ -225,10 +225,11 @@ bool UInterchangeGenericAssetsPipeline::ExecutePostImportPipeline(const UInterch
 {
 	//We do not use the provided base container since ExecutePreImportPipeline cache it
 	//We just make sure the same one is pass in parameter
-	if (!InBaseNodeContainer || !ensure(InBaseNodeContainer == InBaseNodeContainer) || !CreatedAsset)
+	if (!InBaseNodeContainer || !ensure(BaseNodeContainer == InBaseNodeContainer) || !CreatedAsset)
 	{
 		return false;
 	}
+
 	UInterchangeBaseNode* Node = BaseNodeContainer->GetNode(NodeKey);
 	if (!Node)
 	{
