Added a member to InterchangeGenericAssetsPipeline to allow the users to specify which file types should be imported as longlat cubemap.

diff --git a/Engine/Plugins/Runtime/Interchange/Source/Pipelines/Private/InterchangeGenericTexturePipeline.cpp b/Engine/Plugins/Runtime/Interchange/Source/Pipelines/Private/InterchangeGenericTexturePipeline.cpp
--- a/Engine/Plugins/Runtime/Interchange/Source/Pipelines/Private/InterchangeGenericTexturePipeline.cpp
+++ b/Engine/Plugins/Runtime/Interchange/Source/Pipelines/Private/InterchangeGenericTexturePipeline.cpp
@@ -23,6 +23,7 @@
 #include "InterchangeTextureLightProfileNode.h"
 #include "InterchangeTextureLightProfileNode.h"
 #include "InterchangeTextureNode.h"
+#include "Misc/Paths.h"
 
 #if WITH_EDITOR
 #include "NormalMapIdentification.h"
@@ -102,7 +103,23 @@ namespace UE::Interchange::Private
 
 UInterchangeTextureFactoryNode* UInterchangeGenericAssetsPipeline::HandleCreationOfTextureFactoryNode(const UInterchangeTextureNode* TextureNode)
 {
-	return CreateTextureFactoryNode(TextureNode, UE::Interchange::Private::GetDefaultFactoryClassFromTextureNodeClass(TextureNode->GetClass()));
+	UClass* FactoryClass = UE::Interchange::Private::GetDefaultFactoryClassFromTextureNodeClass(TextureNode->GetClass());
+
+#if WITH_EDITORONLY_DATA
+	if (FactoryClass == UInterchangeTextureFactoryNode::StaticClass())
+	{ 
+		if (TOptional<FString> SourceFile = TextureNode->GetPayLoadKey())
+		{
+			const FString Extension = FPaths::GetExtension(SourceFile.GetValue()).ToLower();
+			if (FileExtensionsToImportAsLongLatCubemap.Contains(Extension))
+			{
+				FactoryClass = UInterchangeTextureCubeFactoryNode::StaticClass();
+			}
+		}
+	}
+#endif
+
+	return CreateTextureFactoryNode(TextureNode, FactoryClass);
 }
 
 UInterchangeTextureFactoryNode* UInterchangeGenericAssetsPipeline::CreateTextureFactoryNode(const UInterchangeTextureNode* TextureNode, const TSubclassOf<UInterchangeTextureFactoryNode>& FactorySubclass)
diff --git a/Engine/Plugins/Runtime/Interchange/Source/Pipelines/Public/InterchangeGenericAssetsPipeline.h b/Engine/Plugins/Runtime/Interchange/Source/Pipelines/Public/InterchangeGenericAssetsPipeline.h
--- a/Engine/Plugins/Runtime/Interchange/Source/Pipelines/Public/InterchangeGenericAssetsPipeline.h
+++ b/Engine/Plugins/Runtime/Interchange/Source/Pipelines/Public/InterchangeGenericAssetsPipeline.h
@@ -189,6 +189,10 @@ public:
 	/** If enabled, the texture's green channel will be inverted for normal maps. */
 	UPROPERTY(EditAnywhere, Category = TEXTURES_CATEGORY)
 	bool bFlipNormalMapGreenChannel = false;
+
+	/** Specify the files type that should be imported as long/lat cubemap */
+	UPROPERTY(EditAnywhere, Category = TEXTURES_CATEGORY)
+	TSet<FString> FileExtensionsToImportAsLongLatCubemap = {"hdr"};
 #endif
 
 	// END Pre import pipeline properties
